[PropertyAnimator] Moved plugin

diff --git a/Engine/Plugins/Experimental/PropertyAnimator/Config/BasePropertyAnimator.ini b/Engine/Plugins/Experimental/PropertyAnimator/Config/BasePropertyAnimator.ini
new file mode 100644
--- /dev/null
+++ b/Engine/Plugins/Experimental/PropertyAnimator/Config/BasePropertyAnimator.ini
@@ -0,0 +1,30 @@
+﻿[CoreRedirects]
++ClassRedirects=(OldName="/Script/PropertyController.PropertyControllerWiggle",NewName="/Script/PropertyAnimator.PropertyAnimatorWiggle")
++ClassRedirects=(OldName="/Script/PropertyController.PropertyControllerWave",NewName="/Script/PropertyAnimator.PropertyAnimatorWave")
++ClassRedirects=(OldName="/Script/PropertyController.PropertyControllerTime",NewName="/Script/PropertyAnimator.PropertyAnimatorTime")
++ClassRedirects=(OldName="/Script/PropertyController.PropertyControllerOscillate",NewName="/Script/PropertyAnimator.PropertyAnimatorOscillate")
++ClassRedirects=(OldName="/Script/PropertyController.PropertyControllerFloatBase",NewName="/Script/PropertyAnimator.PropertyAnimatorFloatBase")
++ClassRedirects=(OldName="/Script/PropertyController.PropertyControllerBounce",NewName="/Script/PropertyAnimator.PropertyAnimatorBounce")
++ClassRedirects=(OldName="/Script/PropertyController.PropertyControllerBlink",NewName="/Script/PropertyAnimator.PropertyAnimatorPulse")
++ClassRedirects=(OldName="/Script/PropertyAnimator.PropertyAnimatorBlink",NewName="/Script/PropertyAnimator.PropertyAnimatorPulse")
++ClassRedirects=(OldName="/Script/PropertyController.PropertyControllerBase",NewName="/Script/PropertyAnimator.PropertyAnimatorBase")
++ClassRedirects=(OldName="/Script/PropertyController.PropertyControllerWaveDoubleChannelOverrideContainer",NewName="/Script/PropertyAnimator.PropertyAnimatorWaveDoubleChannelOverrideContainer")
++ClassRedirects=(OldName="/Script/PropertyController.PropertyControllerWaveDoubleChannelEvaluatorSystem",NewName="/Script/PropertyAnimator.PropertyAnimatorWaveDoubleChannelEvaluatorSystem")
++ClassRedirects=(OldName="/Script/PropertyController.PropertyControllerEasingDoubleChannelOverrideContainer",NewName="/Script/PropertyAnimator.PropertyAnimatorEasingDoubleChannelOverrideContainer")
++ClassRedirects=(OldName="/Script/PropertyController.PropertyControllerEasingDoubleChannelEvaluatorSystem",NewName="/Script/PropertyAnimator.PropertyAnimatorEasingDoubleChannelEvaluatorSystem")
++ClassRedirects=(OldName="/Script/PropertyController.PropertyControllerCustomTimeSource",NewName="/Script/PropertyAnimator.PropertyAnimatorCustomTimeSource")
++ClassRedirects=(OldName="/Script/PropertyController.PropertyControllerTimeSourceBase",NewName="/Script/PropertyAnimator.PropertyAnimatorTimeSourceBase")
++ClassRedirects=(OldName="/Script/PropertyController.PropertyControllerWorldTimeSource",NewName="/Script/PropertyAnimator.PropertyAnimatorWorldTimeSource")
++ClassRedirects=(OldName="/Script/AvalanchePropertyAnimator.AvaPropertyAnimatorTextResolver",NewName="/Script/PropertyAnimator.PropertyAnimatorTextResolver")
+
++EnumRedirects=(OldName="/Script/PropertyController.EPropertyControllerEasingType",NewName="/Script/PropertyAnimator.EPropertyAnimatorEasingType")
++EnumRedirects=(OldName="/Script/PropertyController.EPropertyControllerEasingFunction",NewName="/Script/PropertyAnimator.EPropertyAnimatorEasingFunction")
++EnumRedirects=(OldName="/Script/PropertyController.EPropertyControllerWaveFunction",NewName="/Script/PropertyAnimator.EPropertyAnimatorWaveFunction")
+
++StructRedirects=(OldName="/Script/PropertyController.PropertyControllerFloatPropertyOptions",NewName="/Script/PropertyAnimator.PropertyAnimatorFloatPropertyOptions")
++StructRedirects=(OldName="/Script/PropertyController.PropertyControllerEasingDoubleChannel",NewName="/Script/PropertyAnimator.PropertyAnimatorEasingDoubleChannel")
++StructRedirects=(OldName="/Script/PropertyController.PropertyControllerEasingParameters",NewName="/Script/PropertyAnimator.PropertyAnimatorEasingParameters")
++StructRedirects=(OldName="/Script/PropertyController.PropertyControllerWaveDoubleChannel",NewName="/Script/PropertyAnimator.PropertyAnimatorWaveDoubleChannel")
++StructRedirects=(OldName="/Script/PropertyController.PropertyControllerWaveParameters",NewName="/Script/PropertyAnimator.PropertyAnimatorWaveParameters")
+
++PropertyRedirects=(OldName="/Script/PropertyAnimator.PropertyAnimatorFloatBase.Strength",NewName="/Script/PropertyAnimator.PropertyAnimatorFloatBase.GlobalMagnitude")
\ No newline at end of file
diff --git a/Engine/Plugins/Experimental/PropertyAnimator/Source/PropertyAnimator/Private/Animators/PropertyAnimatorBounce.cpp b/Engine/Plugins/Experimental/PropertyAnimator/Source/PropertyAnimator/Private/Animators/PropertyAnimatorBounce.cpp
new file mode 100644
--- /dev/null
+++ b/Engine/Plugins/Experimental/PropertyAnimator/Source/PropertyAnimator/Private/Animators/PropertyAnimatorBounce.cpp
@@ -0,0 +1,36 @@
+// Copyright Epic Games, Inc. All Rights Reserved.
+
+#include "Animators/PropertyAnimatorBounce.h"
+
+#include "Properties/Handlers/PropertyAnimatorCoreHandlerBase.h"
+#include "Properties/PropertyAnimatorFloatContext.h"
+#include "PropertyAnimatorShared.h"
+
+UPropertyAnimatorBounce::UPropertyAnimatorBounce()
+{
+	SetAnimatorDisplayName(DefaultControllerName);
+}
+
+void UPropertyAnimatorBounce::SetInvertEffect(bool bInvert)
+{
+	if (bInvertEffect == bInvert)
+	{
+		return;
+	}
+
+	bInvertEffect = bInvert;
+	OnInvertEffect();
+}
+
+float UPropertyAnimatorBounce::Evaluate(double InTimeElapsed, const FPropertyAnimatorCoreData& InPropertyData, UPropertyAnimatorFloatContext* InOptions) const
+{
+	const float Frequency = InOptions->GetFrequency() * GlobalFrequency;
+	double TimeProgress = FMath::Fmod(InTimeElapsed + InOptions->GetTimeOffset(), 1.f / Frequency);
+
+	TimeProgress = bInvertEffect ? TimeProgress : 1 - TimeProgress;
+
+	const double EasingValue = UE::PropertyAnimator::Easing::Bounce(TimeProgress, EPropertyAnimatorEasingType::In);
+
+	// Remap from [0, 1] to user amplitude from [Min, Max]
+	return FMath::GetMappedRangeValueClamped(FVector2D(0, 1), FVector2D(InOptions->GetAmplitudeMin(), InOptions->GetAmplitudeMax()), EasingValue);
+}
diff --git a/Engine/Plugins/Experimental/PropertyAnimator/Source/PropertyAnimator/Private/Animators/PropertyAnimatorClock.cpp b/Engine/Plugins/Experimental/PropertyAnimator/Source/PropertyAnimator/Private/Animators/PropertyAnimatorClock.cpp
new file mode 100644
--- /dev/null
+++ b/Engine/Plugins/Experimental/PropertyAnimator/Source/PropertyAnimator/Private/Animators/PropertyAnimatorClock.cpp
@@ -0,0 +1,191 @@
+// Copyright Epic Games, Inc. All Rights Reserved.
+
+#include "Animators/PropertyAnimatorClock.h"
+
+#include "Internationalization/Regex.h"
+#include "Misc/DateTime.h"
+#include "Properties/Converters/PropertyAnimatorCoreConverterBase.h"
+#include "Properties/Handlers/PropertyAnimatorCoreHandlerBase.h"
+#include "Subsystems/PropertyAnimatorCoreSubsystem.h"
+
+UPropertyAnimatorClock::UPropertyAnimatorClock()
+{
+	SetAnimatorDisplayName(DefaultControllerName);
+
+	if (!IsTemplate())
+	{
+		OnModeChanged();
+	}
+}
+
+void UPropertyAnimatorClock::SetMode(EPropertyAnimatorClockMode InMode)
+{
+	if (Mode == InMode)
+	{
+		return;
+	}
+
+	Mode = InMode;
+	OnModeChanged();
+}
+
+void UPropertyAnimatorClock::SetDisplayFormat(const FString& InDisplayFormat)
+{
+	DisplayFormat = InDisplayFormat;
+}
+
+void UPropertyAnimatorClock::SetCountdownDuration(const FString& InDisplayFormat)
+{
+	if (DisplayFormat == InDisplayFormat)
+	{
+		return;
+	}
+
+	DisplayFormat = InDisplayFormat;
+	OnModeChanged();
+}
+
+bool UPropertyAnimatorClock::IsPropertyDirectlySupported(const FPropertyAnimatorCoreData& InPropertyData) const
+{
+	return InPropertyData.IsA<FStrProperty>();
+}
+
+bool UPropertyAnimatorClock::IsPropertyIndirectlySupported(const FPropertyAnimatorCoreData& InPropertyData) const
+{
+	// Check if a converter supports the conversion
+	if (UPropertyAnimatorCoreSubsystem* AnimatorSubsystem = UPropertyAnimatorCoreSubsystem::Get())
+	{
+		static const FPropertyBagPropertyDesc AnimatorTypeDesc("", EPropertyBagPropertyType::String);
+		const FPropertyBagPropertyDesc PropertyTypeDesc("", InPropertyData.GetLeafProperty());
+
+		return AnimatorSubsystem->IsConversionSupported(AnimatorTypeDesc, PropertyTypeDesc);
+	}
+
+	return false;
+}
+
+void UPropertyAnimatorClock::PostLoad()
+{
+	Super::PostLoad();
+
+	OnModeChanged();
+}
+
+#if WITH_EDITOR
+void UPropertyAnimatorClock::PostEditChangeProperty(FPropertyChangedEvent& InPropertyChangedEvent)
+{
+	Super::PostEditChangeProperty(InPropertyChangedEvent);
+
+	const FName MemberName = InPropertyChangedEvent.GetMemberPropertyName();
+
+	if (MemberName == GET_MEMBER_NAME_CHECKED(UPropertyAnimatorClock, Mode)
+		|| MemberName == GET_MEMBER_NAME_CHECKED(UPropertyAnimatorClock, CountdownDuration))
+	{
+		OnModeChanged();
+	}
+}
+#endif
+
+void UPropertyAnimatorClock::OnModeChanged()
+{
+	if (Mode == EPropertyAnimatorClockMode::LocalTime)
+	{
+		ActiveTimeSpan = FTimespan(FDateTime::Now().GetTicks()) - ElapsedTimeSpan;
+	}
+	else if (Mode == EPropertyAnimatorClockMode::Countdown)
+	{
+		ActiveTimeSpan = ParseTime(CountdownDuration);
+	}
+	else if (Mode == EPropertyAnimatorClockMode::Stopwatch)
+	{
+		ActiveTimeSpan = FTimespan::Zero();
+	}
+}
+
+void UPropertyAnimatorClock::EvaluateProperties(const FPropertyAnimatorCoreEvaluationParameters& InParameters)
+{
+	FString FormattedDateTime;
+
+	if (Mode == EPropertyAnimatorClockMode::Countdown)
+	{
+		ElapsedTimeSpan = ActiveTimeSpan - FTimespan::FromSeconds(InParameters.TimeElapsed);
+		const FDateTime CountdownDateTime(ElapsedTimeSpan > FTimespan::Zero() ? ElapsedTimeSpan.GetTicks() : 0);
+		FormattedDateTime = CountdownDateTime.ToFormattedString(*DisplayFormat);
+	}
+	else if (Mode == EPropertyAnimatorClockMode::Stopwatch || Mode == EPropertyAnimatorClockMode::LocalTime)
+	{
+		ElapsedTimeSpan = ActiveTimeSpan + FTimespan::FromSeconds(InParameters.TimeElapsed);
+		const FDateTime StopWatchDateTime(ElapsedTimeSpan.GetTicks());
+		FormattedDateTime = StopWatchDateTime.ToFormattedString(*DisplayFormat);
+	}
+
+	EvaluateEachLinkedProperty<UPropertyAnimatorCoreContext>([this, FormattedDateTime](
+		UPropertyAnimatorCoreContext* InContext,
+		const FPropertyAnimatorCoreData& InResolvedProperty,
+		FInstancedPropertyBag& InEvaluatedValues)->bool
+	{
+		const FName DisplayName(InResolvedProperty.GetPathHash());
+
+		InEvaluatedValues.AddProperty(DisplayName, EPropertyBagPropertyType::String);
+		InEvaluatedValues.SetValueString(DisplayName, FormattedDateTime);
+
+		return true;
+	});
+}
+
+void UPropertyAnimatorClock::OnPropertyLinked(UPropertyAnimatorCoreContext* InLinkedProperty)
+{
+	Super::OnPropertyLinked(InLinkedProperty);
+
+	const FPropertyAnimatorCoreData& Property = InLinkedProperty->GetAnimatedProperty();
+	if (Property.IsA<FStrProperty>())
+	{
+		return;
+	}
+
+	const UPropertyAnimatorCoreSubsystem* AnimatorSubsystem = UPropertyAnimatorCoreSubsystem::Get();
+	if (!AnimatorSubsystem)
+	{
+		return;
+	}
+
+	static const FPropertyBagPropertyDesc AnimatorTypeDesc("", EPropertyBagPropertyType::String);
+	const FPropertyBagPropertyDesc PropertyTypeDesc("", Property.GetLeafProperty());
+	const TSet<UPropertyAnimatorCoreConverterBase*> Converters = AnimatorSubsystem->GetSupportedConverters(AnimatorTypeDesc, PropertyTypeDesc);
+	check(!Converters.IsEmpty())
+	InLinkedProperty->SetConverterClass(Converters.Array()[0]->GetClass());
+}
+
+FTimespan UPropertyAnimatorClock::ParseTime(const FString& InFormat)
+{
+	// Regex patterns for different formats
+	static const FRegexPattern HHMMSSPattern(TEXT("^(?:(\\d{2}):)?(\\d{2}):(\\d{2})$")); // 01:00 00:01:00
+	static const FRegexPattern CombinedPattern(TEXT("(?:(\\d+)h)? ?(?:(\\d+)m)? ?(?:(\\d+)s)?")); // 1h 1m 1s
+
+	FRegexMatcher HHMMSSMatcher(HHMMSSPattern, InFormat);
+	FRegexMatcher CombinedMatcher(CombinedPattern, InFormat);
+
+	FTimespan ParsedTimeSpan = FTimespan::Zero();
+
+	if (InFormat.IsNumeric())
+	{
+		const int32 Seconds = FCString::Atoi(*InFormat);
+		ParsedTimeSpan = FTimespan::FromSeconds(Seconds);
+	}
+	else if (HHMMSSMatcher.FindNext())
+	{
+		const int32 Hours = HHMMSSMatcher.GetCaptureGroup(1).IsEmpty() ? 0 : FCString::Atoi(*HHMMSSMatcher.GetCaptureGroup(1));
+		const int32 Minutes = FCString::Atoi(*HHMMSSMatcher.GetCaptureGroup(2));
+		const int32 Seconds = FCString::Atoi(*HHMMSSMatcher.GetCaptureGroup(3));
+		ParsedTimeSpan = FTimespan::FromHours(Hours) + FTimespan::FromMinutes(Minutes) + FTimespan::FromSeconds(Seconds);
+	}
+	else if (CombinedMatcher.FindNext())
+	{
+		const int32 Hours = CombinedMatcher.GetCaptureGroup(1).IsEmpty() ? 0 : FCString::Atoi(*CombinedMatcher.GetCaptureGroup(1));
+		const int32 Minutes = CombinedMatcher.GetCaptureGroup(2).IsEmpty() ? 0 : FCString::Atoi(*CombinedMatcher.GetCaptureGroup(2));
+		const int32 Seconds = CombinedMatcher.GetCaptureGroup(3).IsEmpty() ? 0 : FCString::Atoi(*CombinedMatcher.GetCaptureGroup(3));
+		ParsedTimeSpan = FTimespan::FromHours(Hours) + FTimespan::FromMinutes(Minutes) + FTimespan::FromSeconds(Seconds);
+	}
+
+	return ParsedTimeSpan;
+}
diff --git a/Engine/Plugins/Experimental/PropertyAnimator/Source/PropertyAnimator/Private/Animators/PropertyAnimatorFloatBase.cpp b/Engine/Plugins/Experimental/PropertyAnimator/Source/PropertyAnimator/Private/Animators/PropertyAnimatorFloatBase.cpp
new file mode 100644
--- /dev/null
+++ b/Engine/Plugins/Experimental/PropertyAnimator/Source/PropertyAnimator/Private/Animators/PropertyAnimatorFloatBase.cpp
@@ -0,0 +1,163 @@
+// Copyright Epic Games, Inc. All Rights Reserved.
+
+#include "Animators/PropertyAnimatorFloatBase.h"
+
+#include "Properties/PropertyAnimatorFloatContext.h"
+#include "Properties/Converters/PropertyAnimatorCoreConverterBase.h"
+#include "Properties/Handlers/PropertyAnimatorCoreHandlerBase.h"
+#include "Subsystems/PropertyAnimatorCoreSubsystem.h"
+
+#if WITH_EDITOR
+void UPropertyAnimatorFloatBase::PostEditChangeProperty(FPropertyChangedEvent& InPropertyChangedEvent)
+{
+	Super::PostEditChangeProperty(InPropertyChangedEvent);
+
+	const FName MemberName = InPropertyChangedEvent.GetMemberPropertyName();
+
+	static const FName RandomTimeOffsetName = GET_MEMBER_NAME_CHECKED(UPropertyAnimatorFloatBase, bRandomTimeOffset);
+	static const FName SeedName = GET_MEMBER_NAME_CHECKED(UPropertyAnimatorFloatBase, Seed);
+
+	if (MemberName == SeedName
+		|| MemberName == RandomTimeOffsetName)
+	{
+		OnSeedChanged();
+	}
+}
+#endif
+
+bool UPropertyAnimatorFloatBase::IsPropertyDirectlySupported(const FPropertyAnimatorCoreData& InPropertyData) const
+{
+	return InPropertyData.IsA<FFloatProperty>();
+}
+
+bool UPropertyAnimatorFloatBase::IsPropertyIndirectlySupported(const FPropertyAnimatorCoreData& InPropertyData) const
+{
+	// Check if a converter supports the conversion
+	if (UPropertyAnimatorCoreSubsystem* AnimatorSubsystem = UPropertyAnimatorCoreSubsystem::Get())
+	{
+		static const FPropertyBagPropertyDesc AnimatorTypeDesc("", EPropertyBagPropertyType::Float);
+		const FPropertyBagPropertyDesc PropertyTypeDesc("", InPropertyData.GetLeafProperty());
+
+		return AnimatorSubsystem->IsConversionSupported(AnimatorTypeDesc, PropertyTypeDesc);
+	}
+
+	return false;
+}
+
+void UPropertyAnimatorFloatBase::SetGlobalMagnitude(float InMagnitude)
+{
+	if (FMath::IsNearlyEqual(GlobalMagnitude, InMagnitude))
+	{
+		return;
+	}
+
+	GlobalMagnitude = InMagnitude;
+	OnMagnitudeChanged();
+}
+
+void UPropertyAnimatorFloatBase::SetGlobalFrequency(float InFrequency)
+{
+	if (FMath::IsNearlyEqual(GlobalFrequency, InFrequency))
+	{
+		return;
+	}
+
+	GlobalFrequency = InFrequency;
+	OnGlobalFrequencyChanged();
+}
+
+void UPropertyAnimatorFloatBase::SetAccumulatedTimeOffset(double InOffset)
+{
+	if (FMath::IsNearlyEqual(AccumulatedTimeOffset, InOffset))
+	{
+		return;
+	}
+
+	AccumulatedTimeOffset = InOffset;
+	OnAccumulatedTimeOffsetChanged();
+}
+
+void UPropertyAnimatorFloatBase::SetRandomTimeOffset(bool bInOffset)
+{
+	if (bRandomTimeOffset == bInOffset)
+	{
+		return;
+	}
+
+	bRandomTimeOffset = bInOffset;
+	OnSeedChanged();
+}
+
+void UPropertyAnimatorFloatBase::SetSeed(int32 InSeed)
+{
+	if (Seed == InSeed)
+	{
+		return;
+	}
+
+	Seed = InSeed;
+	OnSeedChanged();
+}
+
+TSubclassOf<UPropertyAnimatorCoreContext> UPropertyAnimatorFloatBase::GetPropertyContextClass(const FPropertyAnimatorCoreData& InUnlinkedProperty)
+{
+	return UPropertyAnimatorFloatContext::StaticClass();
+}
+
+void UPropertyAnimatorFloatBase::EvaluateProperties(const FPropertyAnimatorCoreEvaluationParameters& InParameters)
+{
+	double AccumulatedTimeOffsetRound = InParameters.TimeElapsed;
+	const float AnimatorMagnitude = GlobalMagnitude * InParameters.AnimatorsMagnitude;
+	RandomStream = FRandomStream(Seed);
+
+	EvaluateEachLinkedProperty<UPropertyAnimatorFloatContext>([this, &AccumulatedTimeOffsetRound, &AnimatorMagnitude](
+		UPropertyAnimatorFloatContext* InOptions
+		, const FPropertyAnimatorCoreData& InResolvedProperty
+		, FInstancedPropertyBag& InEvaluatedValues)->bool
+	{
+		const double RandomTimeOffset = bRandomTimeOffset ? RandomStream.GetFraction() : 0;
+		AccumulatedTimeOffsetRound += AccumulatedTimeOffset + RandomTimeOffset;
+
+		if (GlobalMagnitude != 0
+			&& GlobalFrequency != 0
+			&& InOptions->GetMagnitude() != 0
+			&& InOptions->GetFrequency() != 0)
+		{
+			const float EvaluationResult = AnimatorMagnitude
+				* InOptions->GetMagnitude()
+				* Evaluate(AccumulatedTimeOffsetRound, InResolvedProperty, InOptions);
+
+			const FName DisplayName(InResolvedProperty.GetPathHash());
+
+			InEvaluatedValues.AddProperty(DisplayName, EPropertyBagPropertyType::Float);
+			InEvaluatedValues.SetValueFloat(DisplayName, EvaluationResult);
+
+			return true;
+		}
+
+		return false;
+	});
+}
+
+void UPropertyAnimatorFloatBase::OnPropertyLinked(UPropertyAnimatorCoreContext* InLinkedProperty)
+{
+	Super::OnPropertyLinked(InLinkedProperty);
+
+	const FPropertyAnimatorCoreData& Property = InLinkedProperty->GetAnimatedProperty();
+	if (Property.IsA<FFloatProperty>())
+	{
+		return;
+	}
+
+	const UPropertyAnimatorCoreSubsystem* AnimatorSubsystem = UPropertyAnimatorCoreSubsystem::Get();
+	if (!AnimatorSubsystem)
+	{
+		return;
+	}
+
+	static const FPropertyBagPropertyDesc AnimatorTypeDesc("", EPropertyBagPropertyType::Float);
+	const FPropertyBagPropertyDesc PropertyTypeDesc("", Property.GetLeafProperty());
+	const TSet<UPropertyAnimatorCoreConverterBase*> Converters = AnimatorSubsystem->GetSupportedConverters(AnimatorTypeDesc, PropertyTypeDesc);
+	check(!Converters.IsEmpty())
+	InLinkedProperty->SetConverterClass(Converters.Array()[0]->GetClass());
+}
diff --git a/Engine/Plugins/Experimental/PropertyAnimator/Source/PropertyAnimator/Private/Animators/PropertyAnimatorOscillate.cpp b/Engine/Plugins/Experimental/PropertyAnimator/Source/PropertyAnimator/Private/Animators/PropertyAnimatorOscillate.cpp
new file mode 100644
--- /dev/null
+++ b/Engine/Plugins/Experimental/PropertyAnimator/Source/PropertyAnimator/Private/Animators/PropertyAnimatorOscillate.cpp
@@ -0,0 +1,52 @@
+// Copyright Epic Games, Inc. All Rights Reserved.
+
+#include "Animators/PropertyAnimatorOscillate.h"
+
+#include "Properties/PropertyAnimatorFloatContext.h"
+#include "PropertyAnimatorShared.h"
+
+UPropertyAnimatorOscillate::UPropertyAnimatorOscillate()
+{
+	SetAnimatorDisplayName(DefaultControllerName);
+}
+
+void UPropertyAnimatorOscillate::SetOscillateFunction(EPropertyAnimatorOscillateFunction InFunction)
+{
+	OscillateFunction = InFunction;
+}
+
+float UPropertyAnimatorOscillate::Evaluate(double InTimeElapsed, const FPropertyAnimatorCoreData& InPropertyData, UPropertyAnimatorFloatContext* InOptions) const
+{
+	double WaveResult = 0.f;
+
+	using namespace UE::PropertyAnimator;
+
+	const float Frequency = InOptions->GetFrequency() * GlobalFrequency;
+
+	switch (OscillateFunction)
+	{
+		case EPropertyAnimatorOscillateFunction::Sine:
+			WaveResult = Wave::Sine(InTimeElapsed, 1.f, Frequency, InOptions->GetTimeOffset());
+		break;
+		case EPropertyAnimatorOscillateFunction::Cosine:
+			WaveResult = Wave::Cosine(InTimeElapsed, 1.f, Frequency, InOptions->GetTimeOffset());
+		break;
+		case EPropertyAnimatorOscillateFunction::Square:
+			WaveResult = Wave::Square(InTimeElapsed, 1.f, Frequency, InOptions->GetTimeOffset());
+		break;
+		case EPropertyAnimatorOscillateFunction::InvertedSquare:
+			WaveResult = Wave::InvertedSquare(InTimeElapsed, 1.f, Frequency, InOptions->GetTimeOffset());
+		break;
+		case EPropertyAnimatorOscillateFunction::Sawtooth:
+			WaveResult = Wave::Sawtooth(InTimeElapsed, 1.f, Frequency, InOptions->GetTimeOffset());
+		break;
+		case EPropertyAnimatorOscillateFunction::Triangle:
+			WaveResult = Wave::Triangle(InTimeElapsed, 1.f, Frequency, InOptions->GetTimeOffset());
+		break;
+		default:
+			checkNoEntry();
+	}
+
+	// Remap from [-1, 1] to user amplitude from [Min, Max]
+	return FMath::GetMappedRangeValueClamped(FVector2D(-1, 1), FVector2D(InOptions->GetAmplitudeMin(), InOptions->GetAmplitudeMax()), WaveResult);
+}
diff --git a/Engine/Plugins/Experimental/PropertyAnimator/Source/PropertyAnimator/Private/Animators/PropertyAnimatorPulse.cpp b/Engine/Plugins/Experimental/PropertyAnimator/Source/PropertyAnimator/Private/Animators/PropertyAnimatorPulse.cpp
new file mode 100644
--- /dev/null
+++ b/Engine/Plugins/Experimental/PropertyAnimator/Source/PropertyAnimator/Private/Animators/PropertyAnimatorPulse.cpp
@@ -0,0 +1,37 @@
+// Copyright Epic Games, Inc. All Rights Reserved.
+
+#include "Animators/PropertyAnimatorPulse.h"
+
+#include "Properties/PropertyAnimatorFloatContext.h"
+#include "PropertyAnimatorShared.h"
+
+UPropertyAnimatorPulse::UPropertyAnimatorPulse()
+{
+	SetAnimatorDisplayName(DefaultControllerName);
+}
+
+void UPropertyAnimatorPulse::SetEasingFunction(EPropertyAnimatorEasingFunction InEasingFunction)
+{
+	EasingFunction = InEasingFunction;
+}
+
+void UPropertyAnimatorPulse::SetEasingType(EPropertyAnimatorEasingType InEasingType)
+{
+	EasingType = InEasingType;
+}
+
+float UPropertyAnimatorPulse::Evaluate(double InTimeElapsed, const FPropertyAnimatorCoreData& InPropertyData, UPropertyAnimatorFloatContext* InOptions) const
+{
+	const float Frequency = InOptions->GetFrequency() * GlobalFrequency;
+
+	const double WaveResult = UE::PropertyAnimator::Wave::Triangle(InTimeElapsed, 1.f, Frequency, InOptions->GetTimeOffset());
+
+	// Result of wave functions is [-1, 1] -> remap to [0, 1] range for easing functions
+	const float NormalizedWaveProgress = FMath::GetMappedRangeValueClamped(FVector2D(-1, 1), FVector2D(0, 1), WaveResult);
+
+	// Apply easing function on normalized progress
+	const float EasingResult = UE::PropertyAnimator::Easing::Ease(NormalizedWaveProgress, EasingFunction, EasingType);
+
+	// Remap from [0, 1] to user amplitude from [Min, Max]
+	return FMath::GetMappedRangeValueClamped(FVector2D(0, 1), FVector2D(InOptions->GetAmplitudeMin(), InOptions->GetAmplitudeMax()), EasingResult);
+}
diff --git a/Engine/Plugins/Experimental/PropertyAnimator/Source/PropertyAnimator/Private/Animators/PropertyAnimatorSoundWave.cpp b/Engine/Plugins/Experimental/PropertyAnimator/Source/PropertyAnimator/Private/Animators/PropertyAnimatorSoundWave.cpp
new file mode 100644
--- /dev/null
+++ b/Engine/Plugins/Experimental/PropertyAnimator/Source/PropertyAnimator/Private/Animators/PropertyAnimatorSoundWave.cpp
@@ -0,0 +1,75 @@
+// Copyright Epic Games, Inc. All Rights Reserved.
+
+#include "Animators/PropertyAnimatorSoundWave.h"
+
+#include "LoudnessNRT.h"
+#include "Properties/Handlers/PropertyAnimatorCoreHandlerBase.h"
+#include "Properties/PropertyAnimatorFloatContext.h"
+#include "Sound/SoundWave.h"
+
+UPropertyAnimatorSoundWave::UPropertyAnimatorSoundWave()
+{
+	SetAnimatorDisplayName(DefaultControllerName);
+}
+
+void UPropertyAnimatorSoundWave::SetLoop(bool bInLoop)
+{
+	bLoop = bInLoop;
+}
+
+#if WITH_EDITOR
+void UPropertyAnimatorSoundWave::PostEditChangeProperty(FPropertyChangedEvent& PropertyChangedEvent)
+{
+	Super::PostEditChangeProperty(PropertyChangedEvent);
+
+	const FName MemberName = PropertyChangedEvent.GetMemberPropertyName();
+
+	if (MemberName == GET_MEMBER_NAME_CHECKED(UPropertyAnimatorSoundWave, SampledSoundWave))
+	{
+		OnSampledSoundWaveChanged();
+	}
+}
+#endif // WITH_EDITOR
+
+void UPropertyAnimatorSoundWave::OnSampledSoundWaveChanged()
+{
+	if (!AudioAnalyzer)
+	{
+		AudioAnalyzer = NewObject<ULoudnessNRT>(this);
+	}
+
+	AudioAnalyzer->Sound = SampledSoundWave;
+
+#if WITH_EDITOR
+	// Needed to analyse new audio sample
+	FProperty* SoundProperty = FindFProperty<FProperty>(ULoudnessNRT::StaticClass(), GET_MEMBER_NAME_CHECKED(ULoudnessNRT, Sound));
+	FPropertyChangedEvent PropertyChangedEvent(SoundProperty, EPropertyChangeType::ValueSet);
+	AudioAnalyzer->PostEditChangeProperty(PropertyChangedEvent);
+#endif
+}
+
+float UPropertyAnimatorSoundWave::Evaluate(double InTimeElapsed, const FPropertyAnimatorCoreData& InPropertyData, UPropertyAnimatorFloatContext* InOptions) const
+{
+	float Loudness = 0.f;
+
+	if (AudioAnalyzer
+		&& AudioAnalyzer->DurationInSeconds > 0)
+	{
+		float SampleTime = InTimeElapsed + InOptions->GetTimeOffset();
+
+		if ((SampleTime >= 0 && SampleTime <= AudioAnalyzer->DurationInSeconds) || bLoop)
+		{
+			SampleTime = FMath::Fmod(SampleTime, AudioAnalyzer->DurationInSeconds);
+
+			if (SampleTime < 0)
+			{
+				SampleTime = AudioAnalyzer->DurationInSeconds + SampleTime;
+			}
+
+			AudioAnalyzer->GetNormalizedLoudnessAtTime(SampleTime, Loudness);
+		}
+	}
+
+	// Remap from [0, 1] to user amplitude from [Min, Max]
+	return FMath::GetMappedRangeValueClamped(FVector2D(0, 1), FVector2D(InOptions->GetAmplitudeMin(), InOptions->GetAmplitudeMax()), Loudness);
+}
diff --git a/Engine/Plugins/Experimental/PropertyAnimator/Source/PropertyAnimator/Private/Animators/PropertyAnimatorTime.cpp b/Engine/Plugins/Experimental/PropertyAnimator/Source/PropertyAnimator/Private/Animators/PropertyAnimatorTime.cpp
new file mode 100644
--- /dev/null
+++ b/Engine/Plugins/Experimental/PropertyAnimator/Source/PropertyAnimator/Private/Animators/PropertyAnimatorTime.cpp
@@ -0,0 +1,22 @@
+// Copyright Epic Games, Inc. All Rights Reserved.
+
+#include "Animators/PropertyAnimatorTime.h"
+
+#include "Properties/Handlers/PropertyAnimatorCoreHandlerBase.h"
+#include "Properties/PropertyAnimatorFloatContext.h"
+
+UPropertyAnimatorTime::UPropertyAnimatorTime()
+{
+	SetAnimatorDisplayName(DefaultControllerName);
+}
+
+float UPropertyAnimatorTime::Evaluate(double InTimeElapsed, const FPropertyAnimatorCoreData& InPropertyData, UPropertyAnimatorFloatContext* InOptions) const
+{
+	const float Frequency = InOptions->GetFrequency() * GlobalFrequency;
+
+	const double TimePeriod = 1.f / Frequency;
+
+	const double TimeProgress = FMath::Fmod(InTimeElapsed + InOptions->GetTimeOffset(), TimePeriod);
+
+	return FMath::GetMappedRangeValueClamped(FVector2f(0, TimePeriod), FVector2f(InOptions->GetAmplitudeMin(), InOptions->GetAmplitudeMax()), TimeProgress);
+}
diff --git a/Engine/Plugins/Experimental/PropertyAnimator/Source/PropertyAnimator/Private/Animators/PropertyAnimatorWiggle.cpp b/Engine/Plugins/Experimental/PropertyAnimator/Source/PropertyAnimator/Private/Animators/PropertyAnimatorWiggle.cpp
new file mode 100644
--- /dev/null
+++ b/Engine/Plugins/Experimental/PropertyAnimator/Source/PropertyAnimator/Private/Animators/PropertyAnimatorWiggle.cpp
@@ -0,0 +1,28 @@
+// Copyright Epic Games, Inc. All Rights Reserved.
+
+#include "Animators/PropertyAnimatorWiggle.h"
+
+#include "Properties/Handlers/PropertyAnimatorCoreHandlerBase.h"
+#include "Properties/PropertyAnimatorFloatContext.h"
+#include "PropertyAnimatorShared.h"
+
+UPropertyAnimatorWiggle::UPropertyAnimatorWiggle()
+{
+	static int32 SeedIncrement = 0;
+
+	SetAnimatorDisplayName(DefaultControllerName);
+
+	bRandomTimeOffset = true;
+	Seed = SeedIncrement++;
+}
+
+float UPropertyAnimatorWiggle::Evaluate(double InTimeElapsed, const FPropertyAnimatorCoreData& InPropertyData, UPropertyAnimatorFloatContext* InOptions) const
+{
+	const float Frequency = InOptions->GetFrequency() * GlobalFrequency;
+
+	// Apply random wave based on time and frequency
+	const double WaveResult = UE::PropertyAnimator::Wave::Perlin(InTimeElapsed, 1.f, Frequency, InOptions->GetTimeOffset());
+
+	// Remap from [-1, 1] to user amplitude from [Min, Max]
+	return FMath::GetMappedRangeValueClamped(FVector2D(-1, 1), FVector2D(InOptions->GetAmplitudeMin(), InOptions->GetAmplitudeMax()), WaveResult);
+}
diff --git a/Engine/Plugins/Experimental/PropertyAnimator/Source/PropertyAnimator/Private/Presets/PropertyAnimatorPresetLocationYZ.cpp b/Engine/Plugins/Experimental/PropertyAnimator/Source/PropertyAnimator/Private/Presets/PropertyAnimatorPresetLocationYZ.cpp
new file mode 100644
--- /dev/null
+++ b/Engine/Plugins/Experimental/PropertyAnimator/Source/PropertyAnimator/Private/Presets/PropertyAnimatorPresetLocationYZ.cpp
@@ -0,0 +1,36 @@
+// Copyright Epic Games, Inc. All Rights Reserved.
+
+#include "Presets/PropertyAnimatorPresetLocationYZ.h"
+
+#include "Animators/PropertyAnimatorCoreBase.h"
+#include "Properties/PropertyAnimatorFloatContext.h"
+
+void UPropertyAnimatorPresetLocationYZ::GetPresetProperties(const AActor* InActor, const UPropertyAnimatorCoreBase* InAnimator, TSet<FPropertyAnimatorCoreData>& OutProperties) const
+{
+	if (!InActor->GetRootComponent())
+	{
+		return;
+	}
+
+	const FName LocationPropertyName = USceneComponent::GetRelativeLocationPropertyName();
+	FStructProperty* LocationProperty = FindFProperty<FStructProperty>(USceneComponent::StaticClass(), LocationPropertyName);
+	check(LocationProperty);
+
+	OutProperties.Add(FPropertyAnimatorCoreData(InActor->GetRootComponent(), LocationProperty, LocationProperty->Struct->FindPropertyByName(TEXT("Y"))));
+	OutProperties.Add(FPropertyAnimatorCoreData(InActor->GetRootComponent(), LocationProperty, LocationProperty->Struct->FindPropertyByName(TEXT("Z"))));
+}
+
+void UPropertyAnimatorPresetLocationYZ::OnPresetApplied(UPropertyAnimatorCoreBase* InAnimator, const TSet<FPropertyAnimatorCoreData>& InProperties)
+{
+	Super::OnPresetApplied(InAnimator, InProperties);
+
+	for (const FPropertyAnimatorCoreData& Property : InProperties)
+	{
+		if (UPropertyAnimatorFloatContext* Context = InAnimator->GetLinkedPropertyContext<UPropertyAnimatorFloatContext>(Property))
+		{
+			Context->SetMode(EPropertyAnimatorCoreMode::Additive);
+			Context->SetAmplitudeMin(-100.f);
+			Context->SetAmplitudeMax(100.f);
+		}
+	}
+}
diff --git a/Engine/Plugins/Experimental/PropertyAnimator/Source/PropertyAnimator/Private/Presets/PropertyAnimatorPresetLocationYZ.h b/Engine/Plugins/Experimental/PropertyAnimator/Source/PropertyAnimator/Private/Presets/PropertyAnimatorPresetLocationYZ.h
new file mode 100644
--- /dev/null
+++ b/Engine/Plugins/Experimental/PropertyAnimator/Source/PropertyAnimator/Private/Presets/PropertyAnimatorPresetLocationYZ.h
@@ -0,0 +1,29 @@
+// Copyright Epic Games, Inc. All Rights Reserved.
+
+#pragma once
+
+#include "Presets/PropertyAnimatorCorePresetBase.h"
+#include "PropertyAnimatorPresetLocationYZ.generated.h"
+
+class AActor;
+class UPropertyAnimatorCoreBase;
+
+/**
+ * Preset for position properties (Y, Z) on scene component
+ */
+UCLASS(Transient)
+class UPropertyAnimatorPresetLocationYZ : public UPropertyAnimatorCorePresetBase
+{
+	GENERATED_BODY()
+
+public:
+	UPropertyAnimatorPresetLocationYZ()
+		: UPropertyAnimatorCorePresetBase(TEXT("LocationYZ"))
+	{}
+
+protected:
+	//~ Begin UPropertyAnimatorCorePresetBase
+	virtual void GetPresetProperties(const AActor* InActor, const UPropertyAnimatorCoreBase* InAnimator, TSet<FPropertyAnimatorCoreData>& OutProperties) const override;
+	virtual void OnPresetApplied(UPropertyAnimatorCoreBase* InAnimator, const TSet<FPropertyAnimatorCoreData>& InProperties) override;
+	//~ End UPropertyAnimatorCorePresetBase
+};
\ No newline at end of file
diff --git a/Engine/Plugins/Experimental/PropertyAnimator/Source/PropertyAnimator/Private/Presets/PropertyAnimatorPresetTextLocation.cpp b/Engine/Plugins/Experimental/PropertyAnimator/Source/PropertyAnimator/Private/Presets/PropertyAnimatorPresetTextLocation.cpp
new file mode 100644
--- /dev/null
+++ b/Engine/Plugins/Experimental/PropertyAnimator/Source/PropertyAnimator/Private/Presets/PropertyAnimatorPresetTextLocation.cpp
@@ -0,0 +1,44 @@
+// Copyright Epic Games, Inc. All Rights Reserved.
+
+#include "Presets/PropertyAnimatorPresetTextLocation.h"
+
+#include "Animators/PropertyAnimatorCoreBase.h"
+#include "Properties/PropertyAnimatorTextResolver.h"
+#include "Properties/PropertyAnimatorFloatContext.h"
+#include "Text3DComponent.h"
+
+void UPropertyAnimatorPresetTextLocation::GetPresetProperties(const AActor* InActor, const UPropertyAnimatorCoreBase* InAnimator, TSet<FPropertyAnimatorCoreData>& OutProperties) const
+{
+	const UText3DComponent* TextComponent = InActor->FindComponentByClass<UText3DComponent>();
+	if (!TextComponent)
+	{
+		return;
+	}
+
+	USceneComponent* TextRootComponent = TextComponent->GetChildComponent(1);
+	if (!TextRootComponent)
+	{
+		return;
+	}
+
+	const FName LocationPropertyName = USceneComponent::GetRelativeLocationPropertyName();
+	FProperty* LocationProperty = FindFProperty<FProperty>(USceneComponent::StaticClass(), LocationPropertyName);
+	check(LocationProperty);
+
+	OutProperties.Add(FPropertyAnimatorCoreData(TextRootComponent, LocationProperty, nullptr, UPropertyAnimatorTextResolver::StaticClass()));
+}
+
+void UPropertyAnimatorPresetTextLocation::OnPresetApplied(UPropertyAnimatorCoreBase* InAnimator, const TSet<FPropertyAnimatorCoreData>& InProperties)
+{
+	Super::OnPresetApplied(InAnimator, InProperties);
+
+	for (const FPropertyAnimatorCoreData& Property : InProperties)
+	{
+		if (UPropertyAnimatorFloatContext* Context = InAnimator->GetLinkedPropertyContext<UPropertyAnimatorFloatContext>(Property))
+		{
+			Context->SetMode(EPropertyAnimatorCoreMode::Additive);
+			Context->SetAmplitudeMin(-100.f);
+			Context->SetAmplitudeMax(100.f);
+		}
+	}
+}
diff --git a/Engine/Plugins/Experimental/PropertyAnimator/Source/PropertyAnimator/Private/Presets/PropertyAnimatorPresetTextLocation.h b/Engine/Plugins/Experimental/PropertyAnimator/Source/PropertyAnimator/Private/Presets/PropertyAnimatorPresetTextLocation.h
new file mode 100644
--- /dev/null
+++ b/Engine/Plugins/Experimental/PropertyAnimator/Source/PropertyAnimator/Private/Presets/PropertyAnimatorPresetTextLocation.h
@@ -0,0 +1,29 @@
+// Copyright Epic Games, Inc. All Rights Reserved.
+
+#pragma once
+
+#include "Presets/PropertyAnimatorCorePresetBase.h"
+#include "PropertyAnimatorPresetTextLocation.generated.h"
+
+class AActor;
+class UPropertyAnimatorCoreBase;
+
+/**
+ * Preset for text character position properties (X, Y, Z) on scene component
+ */
+UCLASS()
+class UPropertyAnimatorPresetTextLocation : public UPropertyAnimatorCorePresetBase
+{
+	GENERATED_BODY()
+
+public:
+	UPropertyAnimatorPresetTextLocation()
+		: UPropertyAnimatorCorePresetBase(TEXT("TextCharacterLocation"))
+	{}
+
+protected:
+	//~ Begin UPropertyAnimatorCorePresetBase
+	virtual void GetPresetProperties(const AActor* InActor, const UPropertyAnimatorCoreBase* InAnimator, TSet<FPropertyAnimatorCoreData>& OutProperties) const override;
+	virtual void OnPresetApplied(UPropertyAnimatorCoreBase* InAnimator, const TSet<FPropertyAnimatorCoreData>& InProperties) override;
+	//~ End UPropertyAnimatorCorePresetBase
+};
\ No newline at end of file
diff --git a/Engine/Plugins/Experimental/PropertyAnimator/Source/PropertyAnimator/Private/Presets/PropertyAnimatorPresetTextRotation.cpp b/Engine/Plugins/Experimental/PropertyAnimator/Source/PropertyAnimator/Private/Presets/PropertyAnimatorPresetTextRotation.cpp
new file mode 100644
--- /dev/null
+++ b/Engine/Plugins/Experimental/PropertyAnimator/Source/PropertyAnimator/Private/Presets/PropertyAnimatorPresetTextRotation.cpp
@@ -0,0 +1,44 @@
+// Copyright Epic Games, Inc. All Rights Reserved.
+
+#include "Presets/PropertyAnimatorPresetTextRotation.h"
+
+#include "Animators/PropertyAnimatorCoreBase.h"
+#include "Properties/PropertyAnimatorTextResolver.h"
+#include "Properties/PropertyAnimatorFloatContext.h"
+#include "Text3DComponent.h"
+
+void UPropertyAnimatorPresetTextRotation::GetPresetProperties(const AActor* InActor, const UPropertyAnimatorCoreBase* InAnimator, TSet<FPropertyAnimatorCoreData>& OutProperties) const
+{
+	const UText3DComponent* TextComponent = InActor->FindComponentByClass<UText3DComponent>();
+	if (!TextComponent)
+	{
+		return;
+	}
+
+	USceneComponent* TextRootComponent = TextComponent->GetChildComponent(1);
+	if (!TextRootComponent)
+	{
+		return;
+	}
+
+	const FName RotationPropertyName = USceneComponent::GetRelativeRotationPropertyName();
+	FProperty* RotationProperty = FindFProperty<FProperty>(USceneComponent::StaticClass(), RotationPropertyName);
+	check(RotationProperty);
+
+	OutProperties.Add(FPropertyAnimatorCoreData(TextRootComponent, RotationProperty, nullptr, UPropertyAnimatorTextResolver::StaticClass()));
+}
+
+void UPropertyAnimatorPresetTextRotation::OnPresetApplied(UPropertyAnimatorCoreBase* InAnimator, const TSet<FPropertyAnimatorCoreData>& InProperties)
+{
+	Super::OnPresetApplied(InAnimator, InProperties);
+
+	for (const FPropertyAnimatorCoreData& Property : InProperties)
+	{
+		if (UPropertyAnimatorFloatContext* Context = InAnimator->GetLinkedPropertyContext<UPropertyAnimatorFloatContext>(Property))
+		{
+			Context->SetMode(EPropertyAnimatorCoreMode::Absolute);
+			Context->SetAmplitudeMin(-90.f);
+			Context->SetAmplitudeMax(90.f);
+		}
+	}
+}
diff --git a/Engine/Plugins/Experimental/PropertyAnimator/Source/PropertyAnimator/Private/Presets/PropertyAnimatorPresetTextRotation.h b/Engine/Plugins/Experimental/PropertyAnimator/Source/PropertyAnimator/Private/Presets/PropertyAnimatorPresetTextRotation.h
new file mode 100644
--- /dev/null
+++ b/Engine/Plugins/Experimental/PropertyAnimator/Source/PropertyAnimator/Private/Presets/PropertyAnimatorPresetTextRotation.h
@@ -0,0 +1,29 @@
+// Copyright Epic Games, Inc. All Rights Reserved.
+
+#pragma once
+
+#include "Presets/PropertyAnimatorCorePresetBase.h"
+#include "PropertyAnimatorPresetTextRotation.generated.h"
+
+class AActor;
+class UPropertyAnimatorCoreBase;
+
+/**
+ * Preset for text character position properties (Roll, Pitch, Yaw) on scene component
+ */
+UCLASS()
+class UPropertyAnimatorPresetTextRotation : public UPropertyAnimatorCorePresetBase
+{
+	GENERATED_BODY()
+
+public:
+	UPropertyAnimatorPresetTextRotation()
+		: UPropertyAnimatorCorePresetBase(TEXT("TextCharacterRotation"))
+	{}
+
+protected:
+	//~ Begin UPropertyAnimatorCorePresetBase
+	virtual void GetPresetProperties(const AActor* InActor, const UPropertyAnimatorCoreBase* InAnimator, TSet<FPropertyAnimatorCoreData>& OutProperties) const override;
+	virtual void OnPresetApplied(UPropertyAnimatorCoreBase* InAnimator, const TSet<FPropertyAnimatorCoreData>& InProperties) override;
+	//~ End UPropertyAnimatorCorePresetBase
+};
\ No newline at end of file
diff --git a/Engine/Plugins/Experimental/PropertyAnimator/Source/PropertyAnimator/Private/Presets/PropertyAnimatorPresetTextScale.cpp b/Engine/Plugins/Experimental/PropertyAnimator/Source/PropertyAnimator/Private/Presets/PropertyAnimatorPresetTextScale.cpp
new file mode 100644
--- /dev/null
+++ b/Engine/Plugins/Experimental/PropertyAnimator/Source/PropertyAnimator/Private/Presets/PropertyAnimatorPresetTextScale.cpp
@@ -0,0 +1,44 @@
+// Copyright Epic Games, Inc. All Rights Reserved.
+
+#include "Presets/PropertyAnimatorPresetTextScale.h"
+
+#include "Animators/PropertyAnimatorCoreBase.h"
+#include "Properties/PropertyAnimatorTextResolver.h"
+#include "Properties/PropertyAnimatorFloatContext.h"
+#include "Text3DComponent.h"
+
+void UPropertyAnimatorPresetTextScale::GetPresetProperties(const AActor* InActor, const UPropertyAnimatorCoreBase* InAnimator, TSet<FPropertyAnimatorCoreData>& OutProperties) const
+{
+	const UText3DComponent* TextComponent = InActor->FindComponentByClass<UText3DComponent>();
+	if (!TextComponent)
+	{
+		return;
+	}
+
+	USceneComponent* TextRootComponent = TextComponent->GetChildComponent(1);
+	if (!TextRootComponent)
+	{
+		return;
+	}
+
+	const FName ScalePropertyName = USceneComponent::GetRelativeScale3DPropertyName();
+	FProperty* ScaleProperty = FindFProperty<FProperty>(USceneComponent::StaticClass(), ScalePropertyName);
+	check(ScaleProperty);
+
+	OutProperties.Add(FPropertyAnimatorCoreData(TextRootComponent, ScaleProperty, nullptr, UPropertyAnimatorTextResolver::StaticClass()));
+}
+
+void UPropertyAnimatorPresetTextScale::OnPresetApplied(UPropertyAnimatorCoreBase* InAnimator, const TSet<FPropertyAnimatorCoreData>& InProperties)
+{
+	Super::OnPresetApplied(InAnimator, InProperties);
+
+	for (const FPropertyAnimatorCoreData& Property : InProperties)
+	{
+		if (UPropertyAnimatorFloatContext* Context = InAnimator->GetLinkedPropertyContext<UPropertyAnimatorFloatContext>(Property))
+		{
+			Context->SetMode(EPropertyAnimatorCoreMode::Absolute);
+			Context->SetAmplitudeMin(0.f);
+			Context->SetAmplitudeMax(1.f);
+		}
+	}
+}
diff --git a/Engine/Plugins/Experimental/PropertyAnimator/Source/PropertyAnimator/Private/Presets/PropertyAnimatorPresetTextScale.h b/Engine/Plugins/Experimental/PropertyAnimator/Source/PropertyAnimator/Private/Presets/PropertyAnimatorPresetTextScale.h
new file mode 100644
--- /dev/null
+++ b/Engine/Plugins/Experimental/PropertyAnimator/Source/PropertyAnimator/Private/Presets/PropertyAnimatorPresetTextScale.h
@@ -0,0 +1,29 @@
+// Copyright Epic Games, Inc. All Rights Reserved.
+
+#pragma once
+
+#include "Presets/PropertyAnimatorCorePresetBase.h"
+#include "PropertyAnimatorPresetTextScale.generated.h"
+
+class AActor;
+class UPropertyAnimatorCoreBase;
+
+/**
+ * Preset for text character position properties (X, Y, Z) on scene component
+ */
+UCLASS()
+class UPropertyAnimatorPresetTextScale : public UPropertyAnimatorCorePresetBase
+{
+	GENERATED_BODY()
+
+public:
+	UPropertyAnimatorPresetTextScale()
+		: UPropertyAnimatorCorePresetBase(TEXT("TextCharacterScale"))
+	{}
+
+protected:
+	//~ Begin UPropertyAnimatorCorePresetBase
+	virtual void GetPresetProperties(const AActor* InActor, const UPropertyAnimatorCoreBase* InAnimator, TSet<FPropertyAnimatorCoreData>& OutProperties) const override;
+	virtual void OnPresetApplied(UPropertyAnimatorCoreBase* InAnimator, const TSet<FPropertyAnimatorCoreData>& InProperties) override;
+	//~ End UPropertyAnimatorCorePresetBase
+};
\ No newline at end of file
diff --git a/Engine/Plugins/Experimental/PropertyAnimator/Source/PropertyAnimator/Private/Presets/PropertyAnimatorPresetTextVisibility.cpp b/Engine/Plugins/Experimental/PropertyAnimator/Source/PropertyAnimator/Private/Presets/PropertyAnimatorPresetTextVisibility.cpp
new file mode 100644
--- /dev/null
+++ b/Engine/Plugins/Experimental/PropertyAnimator/Source/PropertyAnimator/Private/Presets/PropertyAnimatorPresetTextVisibility.cpp
@@ -0,0 +1,48 @@
+// Copyright Epic Games, Inc. All Rights Reserved.
+
+#include "Presets/PropertyAnimatorPresetTextVisibility.h"
+
+#include "Animators/PropertyAnimatorCoreBase.h"
+#include "Properties/PropertyAnimatorTextResolver.h"
+#include "Properties/Converters/PropertyAnimatorCoreConverterTraits.h"
+#include "Text3DComponent.h"
+
+void UPropertyAnimatorPresetTextVisibility::GetPresetProperties(const AActor* InActor, const UPropertyAnimatorCoreBase* InAnimator, TSet<FPropertyAnimatorCoreData>& OutProperties) const
+{
+	const UText3DComponent* TextComponent = InActor->FindComponentByClass<UText3DComponent>();
+	if (!TextComponent)
+	{
+		return;
+	}
+
+	USceneComponent* TextRootComponent = TextComponent->GetChildComponent(1);
+	if (!TextRootComponent)
+	{
+		return;
+	}
+
+	const FName VisibilityPropertyName = USceneComponent::GetVisiblePropertyName();
+	FProperty* VisibilityProperty = FindFProperty<FProperty>(USceneComponent::StaticClass(), VisibilityPropertyName);
+	check(VisibilityProperty);
+
+	OutProperties.Add(FPropertyAnimatorCoreData(TextRootComponent, VisibilityProperty, nullptr, UPropertyAnimatorTextResolver::StaticClass()));
+}
+
+void UPropertyAnimatorPresetTextVisibility::OnPresetApplied(UPropertyAnimatorCoreBase* InAnimator, const TSet<FPropertyAnimatorCoreData>& InProperties)
+{
+	Super::OnPresetApplied(InAnimator, InProperties);
+
+	for (const FPropertyAnimatorCoreData& Property : InProperties)
+	{
+		if (UPropertyAnimatorCoreContext* Context = InAnimator->GetLinkedPropertyContext<UPropertyAnimatorCoreContext>(Property))
+		{
+			if (FBoolConverterRule* ConverterRule = Context->GetConverterRule<FBoolConverterRule>())
+			{
+				ConverterRule->TrueConditions.Empty();
+
+				FBoolConverterCondition NewCondition(EBoolConverterComparison::Greater, 0);
+				ConverterRule->TrueConditions.Add(NewCondition);
+			}
+		}
+	}
+}
diff --git a/Engine/Plugins/Experimental/PropertyAnimator/Source/PropertyAnimator/Private/Presets/PropertyAnimatorPresetTextVisibility.h b/Engine/Plugins/Experimental/PropertyAnimator/Source/PropertyAnimator/Private/Presets/PropertyAnimatorPresetTextVisibility.h
new file mode 100644
--- /dev/null
+++ b/Engine/Plugins/Experimental/PropertyAnimator/Source/PropertyAnimator/Private/Presets/PropertyAnimatorPresetTextVisibility.h
@@ -0,0 +1,29 @@
+// Copyright Epic Games, Inc. All Rights Reserved.
+
+#pragma once
+
+#include "Presets/PropertyAnimatorCorePresetBase.h"
+#include "PropertyAnimatorPresetTextVisibility.generated.h"
+
+class AActor;
+class UPropertyAnimatorCoreBase;
+
+/**
+ * Preset for text character visibility on scene component
+ */
+UCLASS()
+class UPropertyAnimatorPresetTextVisibility : public UPropertyAnimatorCorePresetBase
+{
+	GENERATED_BODY()
+
+public:
+	UPropertyAnimatorPresetTextVisibility()
+		: UPropertyAnimatorCorePresetBase(TEXT("TextCharacterVisibility"))
+	{}
+
+protected:
+	//~ Begin UPropertyAnimatorCorePresetBase
+	virtual void GetPresetProperties(const AActor* InActor, const UPropertyAnimatorCoreBase* InAnimator, TSet<FPropertyAnimatorCoreData>& OutProperties) const override;
+	virtual void OnPresetApplied(UPropertyAnimatorCoreBase* InAnimator, const TSet<FPropertyAnimatorCoreData>& InProperties) override;
+	//~ End UPropertyAnimatorCorePresetBase
+};
\ No newline at end of file
diff --git a/Engine/Plugins/Experimental/PropertyAnimator/Source/PropertyAnimator/Private/Properties/PropertyAnimatorTextGroup.cpp b/Engine/Plugins/Experimental/PropertyAnimator/Source/PropertyAnimator/Private/Properties/PropertyAnimatorTextGroup.cpp
new file mode 100644
--- /dev/null
+++ b/Engine/Plugins/Experimental/PropertyAnimator/Source/PropertyAnimator/Private/Properties/PropertyAnimatorTextGroup.cpp
@@ -0,0 +1,51 @@
+// Copyright Epic Games, Inc. All Rights Reserved.
+
+#include "Properties/PropertyAnimatorTextGroup.h"
+
+#include "Properties/PropertyAnimatorCoreContext.h"
+#include "Properties/PropertyAnimatorCoreResolver.h"
+#include "Text3DComponent.h"
+
+void UPropertyAnimatorTextGroup::ManageProperties(const UPropertyAnimatorCoreContext* InContext, TArray<FPropertyAnimatorCoreData>& InOutProperties)
+{
+	if (InOutProperties.IsEmpty())
+	{
+		return;
+	}
+
+	const int32 MaxIndex = InOutProperties.Num();
+	const int32 BeginIndex = RangeStart * MaxIndex + RangeOffset * MaxIndex;
+	const int32 EndIndex = RangeEnd * MaxIndex + RangeOffset * MaxIndex;
+
+	if (EndIndex < 0 || BeginIndex > EndIndex || BeginIndex == EndIndex || BeginIndex > MaxIndex)
+	{
+		InOutProperties.Empty();
+		return;
+	}
+
+	if (EndIndex < MaxIndex)
+	{
+		InOutProperties.RemoveAt(EndIndex, MaxIndex - EndIndex);
+	}
+
+	if (BeginIndex > 0 && BeginIndex <= MaxIndex)
+	{
+		InOutProperties.RemoveAt(0, BeginIndex);
+	}
+}
+
+bool UPropertyAnimatorTextGroup::IsPropertySupported(const UPropertyAnimatorCoreContext* InContext) const
+{
+	const UActorComponent* ActorComponent = InContext->GetAnimatedProperty().GetOwningComponent();
+	const UPropertyAnimatorCoreResolver* PropertyResolver = InContext->GetAnimatedProperty().GetPropertyResolver();
+
+	if (PropertyResolver
+		&& PropertyResolver->IsRange()
+		&& ActorComponent
+		&& (ActorComponent->IsA<UText3DComponent>() || ActorComponent->GetTypedOuter<UText3DComponent>()))
+	{
+		return true;
+	}
+
+	return Super::IsPropertySupported(InContext);
+}
diff --git a/Engine/Plugins/Experimental/PropertyAnimator/Source/PropertyAnimator/Private/Properties/PropertyAnimatorTextResolver.cpp b/Engine/Plugins/Experimental/PropertyAnimator/Source/PropertyAnimator/Private/Properties/PropertyAnimatorTextResolver.cpp
new file mode 100644
--- /dev/null
+++ b/Engine/Plugins/Experimental/PropertyAnimator/Source/PropertyAnimator/Private/Properties/PropertyAnimatorTextResolver.cpp
@@ -0,0 +1,72 @@
+// Copyright Epic Games, Inc. All Rights Reserved.
+
+#include "Properties/PropertyAnimatorTextResolver.h"
+
+#include "Text3DComponent.h"
+
+void UPropertyAnimatorTextResolver::GetResolvableProperties(const FPropertyAnimatorCoreData& InParentProperty, TSet<FPropertyAnimatorCoreData>& OutProperties)
+{
+	const AActor* Actor = InParentProperty.GetOwningActor();
+	if (!Actor || InParentProperty.IsResolved())
+	{
+		return;
+	}
+
+	UText3DComponent* TextComponent = Actor->FindComponentByClass<UText3DComponent>();
+	if (!TextComponent)
+	{
+		return;
+	}
+
+	USceneComponent* TextRootComponent = TextComponent->GetChildComponent(1);
+	if (!TextRootComponent)
+	{
+		return;
+	}
+
+	FProperty* RelativeLocation = FindFProperty<FProperty>(TextRootComponent->GetClass(), TEXT("RelativeLocation"));
+	FPropertyAnimatorCoreData RelativeLocationProperty(TextRootComponent, RelativeLocation, nullptr, GetClass());
+	OutProperties.Add(RelativeLocationProperty);
+
+	FProperty* RelativeRotation = FindFProperty<FProperty>(TextRootComponent->GetClass(), TEXT("RelativeRotation"));
+	FPropertyAnimatorCoreData RelativeRotationProperty(TextRootComponent, RelativeRotation, nullptr, GetClass());
+	OutProperties.Add(RelativeRotationProperty);
+
+	FProperty* RelativeScale = FindFProperty<FProperty>(TextRootComponent->GetClass(), TEXT("RelativeScale3D"));
+	FPropertyAnimatorCoreData RelativeScaleProperty(TextRootComponent, RelativeScale, nullptr, GetClass());
+	OutProperties.Add(RelativeScaleProperty);
+}
+
+void UPropertyAnimatorTextResolver::ResolveProperties(const FPropertyAnimatorCoreData& InTemplateProperty, TArray<FPropertyAnimatorCoreData>& OutProperties)
+{
+	if (!InTemplateProperty.IsResolvable())
+	{
+		return;
+	}
+
+	const USceneComponent* TextRootComponent = Cast<USceneComponent>(InTemplateProperty.GetOwningComponent());
+	if (!TextRootComponent)
+	{
+		return;
+	}
+
+	TArray<FProperty*> ChainProperties;
+	for (const TFieldPath<FProperty>& Property : InTemplateProperty.GetChainProperties())
+	{
+		ChainProperties.Add(Property.Get());
+	}
+
+	// Gather each character in the text
+	for (int32 ComponentIndex = 0; ComponentIndex < TextRootComponent->GetNumChildrenComponents(); ComponentIndex++)
+	{
+		USceneComponent* CharacterKerningComponent = TextRootComponent->GetChildComponent(ComponentIndex);
+
+		if (!CharacterKerningComponent)
+		{
+			continue;
+		}
+
+		FPropertyAnimatorCoreData CharacterProperty(CharacterKerningComponent, ChainProperties);
+		OutProperties.Add(CharacterProperty);
+	}
+}
diff --git a/Engine/Plugins/Experimental/PropertyAnimator/Source/PropertyAnimator/Private/Properties/PropertyAnimatorTextResolver.h b/Engine/Plugins/Experimental/PropertyAnimator/Source/PropertyAnimator/Private/Properties/PropertyAnimatorTextResolver.h
new file mode 100644
--- /dev/null
+++ b/Engine/Plugins/Experimental/PropertyAnimator/Source/PropertyAnimator/Private/Properties/PropertyAnimatorTextResolver.h
@@ -0,0 +1,29 @@
+// Copyright Epic Games, Inc. All Rights Reserved.
+
+#pragma once
+
+#include "Properties/PropertyAnimatorCoreData.h"
+#include "Properties/PropertyAnimatorCoreResolver.h"
+#include "PropertyAnimatorTextResolver.generated.h"
+
+/**
+ * Text characters properties resolver
+ * Since each character in text are transient and regenerated on change
+ * We need to have a resolvable property that will be resolved to each characters in the text we needed
+ * We manipulate that resolvable property that underneath means we manipulate all text characters properties
+ */
+UCLASS(Transient)
+class UPropertyAnimatorTextResolver : public UPropertyAnimatorCoreResolver
+{
+	GENERATED_BODY()
+
+public:
+	UPropertyAnimatorTextResolver()
+		: UPropertyAnimatorCoreResolver(TEXT("TextCharacters"), /** IsRange */ true)
+	{}
+
+	//~ Begin UPropertyAnimatorCoreResolver
+	virtual void GetResolvableProperties(const FPropertyAnimatorCoreData& InParentProperty, TSet<FPropertyAnimatorCoreData>& OutProperties) override;
+	virtual void ResolveProperties(const FPropertyAnimatorCoreData& InTemplateProperty, TArray<FPropertyAnimatorCoreData>& OutProperties) override;
+	//~ End UPropertyAnimatorCoreResolver
+};
\ No newline at end of file
diff --git a/Engine/Plugins/Experimental/PropertyAnimator/Source/PropertyAnimator/Private/PropertyAnimatorModule.cpp b/Engine/Plugins/Experimental/PropertyAnimator/Source/PropertyAnimator/Private/PropertyAnimatorModule.cpp
new file mode 100644
--- /dev/null
+++ b/Engine/Plugins/Experimental/PropertyAnimator/Source/PropertyAnimator/Private/PropertyAnimatorModule.cpp
@@ -0,0 +1,5 @@
+﻿// Copyright Epic Games, Inc. All Rights Reserved.
+
+#include "Modules/ModuleManager.h"
+
+IMPLEMENT_MODULE(FDefaultModuleImpl, PropertyAnimator)
diff --git a/Engine/Plugins/Experimental/PropertyAnimator/Source/PropertyAnimator/Public/Animators/PropertyAnimatorBounce.h b/Engine/Plugins/Experimental/PropertyAnimator/Source/PropertyAnimator/Public/Animators/PropertyAnimatorBounce.h
new file mode 100644
--- /dev/null
+++ b/Engine/Plugins/Experimental/PropertyAnimator/Source/PropertyAnimator/Public/Animators/PropertyAnimatorBounce.h
@@ -0,0 +1,37 @@
+// Copyright Epic Games, Inc. All Rights Reserved.
+
+#pragma once
+
+#include "Animators/PropertyAnimatorFloatBase.h"
+#include "PropertyAnimatorBounce.generated.h"
+
+/**
+ * Applies an additive bounce movement with various options on supported float properties
+ */
+UCLASS(AutoExpandCategories=("Animator"))
+class PROPERTYANIMATOR_API UPropertyAnimatorBounce : public UPropertyAnimatorFloatBase
+{
+	GENERATED_BODY()
+
+public:
+	static constexpr const TCHAR* DefaultControllerName = TEXT("Bounce");
+
+	UPropertyAnimatorBounce();
+
+	void SetInvertEffect(bool bInvert);
+	bool GetInvertEffect() const
+	{
+		return bInvertEffect;
+	}
+
+protected:
+	virtual void OnInvertEffect() {}
+
+	//~ Begin UPropertyAnimatorFloatBase
+	virtual float Evaluate(double InTimeElapsed, const FPropertyAnimatorCoreData& InPropertyData, UPropertyAnimatorFloatContext* InOptions) const override;
+	//~ End UPropertyAnimatorFloatBase
+
+	/** Invert the effect result */
+	UPROPERTY(EditAnywhere, BlueprintReadWrite, Setter="SetInvertEffect", Getter="GetInvertEffect", Category="Animator")
+	bool bInvertEffect = false;
+};
\ No newline at end of file
diff --git a/Engine/Plugins/Experimental/PropertyAnimator/Source/PropertyAnimator/Public/Animators/PropertyAnimatorClock.h b/Engine/Plugins/Experimental/PropertyAnimator/Source/PropertyAnimator/Public/Animators/PropertyAnimatorClock.h
new file mode 100644
--- /dev/null
+++ b/Engine/Plugins/Experimental/PropertyAnimator/Source/PropertyAnimator/Public/Animators/PropertyAnimatorClock.h
@@ -0,0 +1,115 @@
+// Copyright Epic Games, Inc. All Rights Reserved.
+
+#pragma once
+
+#include "Animators/PropertyAnimatorCoreBase.h"
+#include "Misc/Timespan.h"
+#include "PropertyAnimatorClock.generated.h"
+
+/** Mode supported for properties value */
+UENUM(BlueprintType)
+enum class EPropertyAnimatorClockMode : uint8
+{
+	/** Local time of the machine */
+	LocalTime,
+	/** Specified duration elapsing until it reaches 0 */
+	Countdown,
+	/** Shows the current time elapsed */
+	Stopwatch
+};
+
+/** Animate supported string properties with a clock feature */
+UCLASS(AutoExpandCategories=("Animator"))
+class PROPERTYANIMATOR_API UPropertyAnimatorClock : public UPropertyAnimatorCoreBase
+{
+	GENERATED_BODY()
+
+public:
+	static constexpr const TCHAR* DefaultControllerName = TEXT("Clock");
+
+	UPropertyAnimatorClock();
+
+	void SetMode(EPropertyAnimatorClockMode InMode);
+	EPropertyAnimatorClockMode GetMode() const
+	{
+		return Mode;
+	}
+
+	void SetDisplayFormat(const FString& InDisplayFormat);
+	const FString& GetDisplayFormat() const
+	{
+		return DisplayFormat;
+	}
+
+	void SetCountdownDuration(const FString& InDisplayFormat);
+	const FString& GetCountdownDuration() const
+	{
+		return CountdownDuration;
+	}
+
+protected:
+	static FTimespan ParseTime(const FString& InFormat);
+
+	//~ Begin UObject
+	virtual void PostLoad() override;
+#if WITH_EDITOR
+	virtual void PostEditChangeProperty(FPropertyChangedEvent& InPropertyChangedEvent) override;
+#endif
+	//~ End UObject
+
+	//~ Begin UPropertyAnimatorCoreBase
+	virtual bool IsPropertyDirectlySupported(const FPropertyAnimatorCoreData& InPropertyData) const override;
+	virtual bool IsPropertyIndirectlySupported(const FPropertyAnimatorCoreData& InPropertyData) const override;
+	virtual void EvaluateProperties(const FPropertyAnimatorCoreEvaluationParameters& InParameters) override;
+	virtual void OnPropertyLinked(UPropertyAnimatorCoreContext* InLinkedProperty) override;
+	//~ End UPropertyAnimatorCoreBase
+
+	void OnModeChanged();
+
+	/** Mode chosen for this clock animator */
+	UPROPERTY(EditInstanceOnly, BlueprintReadWrite, Setter, Getter, Category="Animator")
+	EPropertyAnimatorClockMode Mode = EPropertyAnimatorClockMode::LocalTime;
+
+	/**
+	 * Display date time format : 
+	 * %a - Weekday, eg) Sun
+	 * %A - Weekday, eg) Sunday
+	 * %w - Weekday, 0-6 (Sunday is 0)
+	 * %y - Year, YY
+	 * %Y - Year, YYYY
+	 * %b - Month, eg) Jan
+	 * %B - Month, eg) January
+	 * %m - Month, 01-12
+	 * %d - Day, 01-31
+	 * %e - Day, 1-31
+	 * %l - 12h Hour, 1-12
+	 * %I - 12h Hour, 01-12
+	 * %H - 24h Hour, 00-23
+	 * %M - Minute, 00-59
+	 * %S - Second, 00-60
+	 * %p - AM or PM
+	 * %P - am or PM
+	 * %j - Day of the Year, 001-366
+	 */
+	UPROPERTY(EditInstanceOnly, BlueprintReadWrite, Setter, Getter, Category="Animator")
+	FString DisplayFormat = TEXT("%H:%M:%S");
+
+	/**
+	 * Countdown duration format : 
+	 * 120 = 2 minutes
+	 * 02:00 = 2 minutes
+	 * 00:02:00 = 2 minutes
+	 * 2m = 2 minutes
+	 * 1h = 1 hour
+	 * 120s = 2 minutes
+	 */
+	UPROPERTY(EditInstanceOnly, BlueprintReadWrite, Setter, Getter, Category="Animator", meta=(EditCondition="Mode == EPropertyAnimatorClockMode::Countdown", EditConditionHides))
+	FString CountdownDuration = TEXT("5m");
+
+private:
+	UPROPERTY(Transient)
+	FTimespan ActiveTimeSpan;
+
+	UPROPERTY(Transient)
+	FTimespan ElapsedTimeSpan;
+};
\ No newline at end of file
diff --git a/Engine/Plugins/Experimental/PropertyAnimator/Source/PropertyAnimator/Public/Animators/PropertyAnimatorFloatBase.h b/Engine/Plugins/Experimental/PropertyAnimator/Source/PropertyAnimator/Public/Animators/PropertyAnimatorFloatBase.h
new file mode 100644
--- /dev/null
+++ b/Engine/Plugins/Experimental/PropertyAnimator/Source/PropertyAnimator/Public/Animators/PropertyAnimatorFloatBase.h
@@ -0,0 +1,100 @@
+// Copyright Epic Games, Inc. All Rights Reserved.
+
+#pragma once
+
+#include "Animators/PropertyAnimatorCoreBase.h"
+#include "PropertyAnimatorFloatBase.generated.h"
+
+class UPropertyAnimatorFloatContext;
+
+/**
+ * Animate supported float properties with various options
+ */
+UCLASS(Abstract, AutoExpandCategories=("Animator"))
+class PROPERTYANIMATOR_API UPropertyAnimatorFloatBase : public UPropertyAnimatorCoreBase
+{
+	GENERATED_BODY()
+
+	friend class FPropertyAnimatorCoreEditorDetailCustomization;
+
+public:
+	void SetGlobalMagnitude(float InMagnitude);
+	float GetGlobalMagnitude() const
+	{
+		return GlobalMagnitude;
+	}
+
+	void SetGlobalFrequency(float InFrequency);
+	float GetGlobalFrequency() const
+	{
+		return GlobalFrequency;
+	}
+
+	void SetAccumulatedTimeOffset(double InOffset);
+	double GetAccumulatedTimeOffset() const
+	{
+		return AccumulatedTimeOffset;
+	}
+
+	void SetRandomTimeOffset(bool bInOffset);
+	bool GetRandomTimeOffset() const
+	{
+		return bRandomTimeOffset;
+	}
+
+	void SetSeed(int32 InSeed);
+	int32 GetSeed() const
+	{
+		return Seed;
+	}
+
+protected:
+	//~ Begin UObject
+#if WITH_EDITOR
+	virtual void PostEditChangeProperty(FPropertyChangedEvent& InPropertyChangedEvent) override;
+#endif
+	//~ End UObject
+
+	virtual void OnMagnitudeChanged() {}
+	virtual void OnGlobalFrequencyChanged() {}
+	virtual void OnAccumulatedTimeOffsetChanged() {}
+	virtual void OnSeedChanged() {}
+
+	//~ Begin UPropertyAnimatorCoreBase
+	virtual TSubclassOf<UPropertyAnimatorCoreContext> GetPropertyContextClass(const FPropertyAnimatorCoreData& InUnlinkedProperty) override;
+	virtual bool IsPropertyDirectlySupported(const FPropertyAnimatorCoreData& InPropertyData) const override;
+	virtual bool IsPropertyIndirectlySupported(const FPropertyAnimatorCoreData& InPropertyData) const override;
+	virtual void EvaluateProperties(const FPropertyAnimatorCoreEvaluationParameters& InParameters) override;
+	virtual void OnPropertyLinked(UPropertyAnimatorCoreContext* InLinkedProperty) override;
+	//~ End UPropertyAnimatorCoreBase
+
+	/** Evaluate and return float value for a property */
+	virtual float Evaluate(double InTimeElapsed, const FPropertyAnimatorCoreData& InPropertyData, UPropertyAnimatorFloatContext* InOptions) const
+	{
+		return 0.f;
+	}
+
+	/** Global magnitude for the effect */
+	UPROPERTY(EditInstanceOnly, BlueprintReadWrite, Setter, Getter, Category="Animator", meta=(ClampMin="0"))
+	float GlobalMagnitude = 1.f;
+
+	/** Global frequency multiplier */
+	UPROPERTY(EditInstanceOnly, BlueprintReadWrite, Setter, Getter, Category="Animator", meta=(ClampMin="0"))
+	float GlobalFrequency = 1.f;
+
+	/** Use random time offset to add variation in animation */
+	UPROPERTY(EditAnywhere, BlueprintReadWrite, Setter="SetRandomTimeOffset", Getter="GetRandomTimeOffset", Category="Animator")
+	bool bRandomTimeOffset = false;
+
+	/** Seed to generate per property time offset */
+	UPROPERTY(EditAnywhere, BlueprintReadWrite, Setter, Getter, Category="Animator", meta=(EditCondition="bRandomTimeOffset", EditConditionHides))
+	int32 Seed = 0;
+
+	/** This time offset will be accumulated for each property for every round */
+	UPROPERTY(EditInstanceOnly, BlueprintReadWrite, Setter, Getter, Category="Animator")
+	double AccumulatedTimeOffset = 0.f;
+
+private:
+	/** Random stream for time offset */
+	FRandomStream RandomStream = FRandomStream(Seed);
+};
\ No newline at end of file
diff --git a/Engine/Plugins/Experimental/PropertyAnimator/Source/PropertyAnimator/Public/Animators/PropertyAnimatorOscillate.h b/Engine/Plugins/Experimental/PropertyAnimator/Source/PropertyAnimator/Public/Animators/PropertyAnimatorOscillate.h
new file mode 100644
--- /dev/null
+++ b/Engine/Plugins/Experimental/PropertyAnimator/Source/PropertyAnimator/Public/Animators/PropertyAnimatorOscillate.h
@@ -0,0 +1,46 @@
+// Copyright Epic Games, Inc. All Rights Reserved.
+
+#pragma once
+
+#include "Animators/PropertyAnimatorFloatBase.h"
+#include "PropertyAnimatorOscillate.generated.h"
+
+UENUM(BlueprintType)
+enum class EPropertyAnimatorOscillateFunction : uint8
+{
+	Sine,
+	Cosine,
+	Square,
+	InvertedSquare,
+	Sawtooth,
+	Triangle,
+};
+
+/**
+ * Applies an additive regular oscillate movement with various options on supported float properties
+ */
+UCLASS(AutoExpandCategories=("Animator"))
+class PROPERTYANIMATOR_API UPropertyAnimatorOscillate : public UPropertyAnimatorFloatBase
+{
+	GENERATED_BODY()
+
+public:
+	static constexpr const TCHAR* DefaultControllerName = TEXT("Oscillate");
+
+	UPropertyAnimatorOscillate();
+
+	void SetOscillateFunction(EPropertyAnimatorOscillateFunction InFunction);
+	EPropertyAnimatorOscillateFunction GetOscillateFunction() const
+	{
+		return OscillateFunction;
+	}
+
+protected:
+	//~ Begin UPropertyAnimatorFloatBase
+	virtual float Evaluate(double InTimeElapsed, const FPropertyAnimatorCoreData& InPropertyData, UPropertyAnimatorFloatContext* InOptions) const override;
+	//~ End UPropertyAnimatorFloatBase
+
+	/** The oscillate function to feed current time elapsed */
+	UPROPERTY(EditAnywhere, BlueprintReadWrite, Setter, Getter, Category="Animator")
+	EPropertyAnimatorOscillateFunction OscillateFunction = EPropertyAnimatorOscillateFunction::Sine;
+};
\ No newline at end of file
diff --git a/Engine/Plugins/Experimental/PropertyAnimator/Source/PropertyAnimator/Public/Animators/PropertyAnimatorPulse.h b/Engine/Plugins/Experimental/PropertyAnimator/Source/PropertyAnimator/Public/Animators/PropertyAnimatorPulse.h
new file mode 100644
--- /dev/null
+++ b/Engine/Plugins/Experimental/PropertyAnimator/Source/PropertyAnimator/Public/Animators/PropertyAnimatorPulse.h
@@ -0,0 +1,48 @@
+// Copyright Epic Games, Inc. All Rights Reserved.
+
+/** The wave function to feed current time elapsed */
+
+#pragma once
+
+#include "Animators/PropertyAnimatorFloatBase.h"
+#include "PropertyAnimatorShared.h"
+#include "PropertyAnimatorPulse.generated.h"
+
+/**
+ * Applies an additive pulse movement with various options on supported float properties
+ */
+UCLASS(AutoExpandCategories=("Animator"))
+class PROPERTYANIMATOR_API UPropertyAnimatorPulse : public UPropertyAnimatorFloatBase
+{
+	GENERATED_BODY()
+
+public:
+	static constexpr const TCHAR* DefaultControllerName = TEXT("Pulse");
+
+	UPropertyAnimatorPulse();
+
+	void SetEasingFunction(EPropertyAnimatorEasingFunction InEasingFunction);
+	EPropertyAnimatorEasingFunction GetEasingFunction() const
+	{
+		return EasingFunction;
+	}
+
+	void SetEasingType(EPropertyAnimatorEasingType InEasingType);
+	EPropertyAnimatorEasingType GetEasingType() const
+	{
+		return EasingType;
+	}
+
+protected:
+	//~ Begin UPropertyAnimatorFloatBase
+	virtual float Evaluate(double InTimeElapsed, const FPropertyAnimatorCoreData& InPropertyData, UPropertyAnimatorFloatContext* InOptions) const override;
+	//~ End UPropertyAnimatorFloatBase
+
+	/** The easing function to use to modify the base effect */
+	UPROPERTY(EditAnywhere, BlueprintReadWrite, Setter, Getter, Category="Animator")
+	EPropertyAnimatorEasingFunction EasingFunction = EPropertyAnimatorEasingFunction::Linear;
+
+	/** The type of effect for easing function */
+	UPROPERTY(EditAnywhere, BlueprintReadWrite, Setter, Getter, Category="Animator", meta=(EditCondition="EasingFunction != EPropertyAnimatorEasingFunction::Line", EditConditionHides))
+	EPropertyAnimatorEasingType EasingType = EPropertyAnimatorEasingType::InOut;
+};
\ No newline at end of file
diff --git a/Engine/Plugins/Experimental/PropertyAnimator/Source/PropertyAnimator/Public/Animators/PropertyAnimatorSoundWave.h b/Engine/Plugins/Experimental/PropertyAnimator/Source/PropertyAnimator/Public/Animators/PropertyAnimatorSoundWave.h
new file mode 100644
--- /dev/null
+++ b/Engine/Plugins/Experimental/PropertyAnimator/Source/PropertyAnimator/Public/Animators/PropertyAnimatorSoundWave.h
@@ -0,0 +1,59 @@
+// Copyright Epic Games, Inc. All Rights Reserved.
+
+#pragma once
+
+#include "Animators/PropertyAnimatorFloatBase.h"
+#include "PropertyAnimatorSoundWave.generated.h"
+
+class ULoudnessNRT;
+class USoundWave;
+
+/**
+ * Applies a sampled sound wave movement with various options on supported float properties
+ */
+UCLASS(AutoExpandCategories=("Animator"))
+class PROPERTYANIMATOR_API UPropertyAnimatorSoundWave : public UPropertyAnimatorFloatBase
+{
+	GENERATED_BODY()
+
+public:
+	static constexpr const TCHAR* DefaultControllerName = TEXT("SoundWave");
+
+	UPropertyAnimatorSoundWave();
+
+	void SetLoop(bool bInLoop);
+	bool GetLoop() const
+	{
+		return bLoop;
+	}
+
+protected:
+	//~ Begin UObject
+#if WITH_EDITOR
+	virtual void PostEditChangeProperty(FPropertyChangedEvent& PropertyChangedEvent) override;
+#endif
+	//~ End UObject
+
+	void OnSampledSoundWaveChanged();
+
+	//~ Begin UPropertyAnimatorFloatBase
+	virtual float Evaluate(double InTimeElapsed, const FPropertyAnimatorCoreData& InPropertyData, UPropertyAnimatorFloatContext* InOptions) const override;
+	//~ End UPropertyAnimatorFloatBase
+
+	/**
+	 * The sound wave to analyse
+	 * Cannot be switched at runtime, only in editor due to analyzer
+	 * Analyzed audio will work at runtime since it is cached
+	 */
+	UPROPERTY(EditInstanceOnly, BlueprintReadOnly, Category="Animator")
+	TObjectPtr<USoundWave> SampledSoundWave;
+
+	/** Whether we keep looping after the duration has been reached or before 0 */
+	UPROPERTY(EditInstanceOnly, BlueprintReadWrite, Setter="SetLoop", Getter="GetLoop", Category="Animator")
+	bool bLoop = true;
+
+private:
+	/** Non-Real-Time audio analyser for loudness */
+	UPROPERTY()
+	TObjectPtr<ULoudnessNRT> AudioAnalyzer;
+};
\ No newline at end of file
diff --git a/Engine/Plugins/Experimental/PropertyAnimator/Source/PropertyAnimator/Public/Animators/PropertyAnimatorTime.h b/Engine/Plugins/Experimental/PropertyAnimator/Source/PropertyAnimator/Public/Animators/PropertyAnimatorTime.h
new file mode 100644
--- /dev/null
+++ b/Engine/Plugins/Experimental/PropertyAnimator/Source/PropertyAnimator/Public/Animators/PropertyAnimatorTime.h
@@ -0,0 +1,25 @@
+// Copyright Epic Games, Inc. All Rights Reserved.
+
+#pragma once
+
+#include "Animators/PropertyAnimatorFloatBase.h"
+#include "PropertyAnimatorTime.generated.h"
+
+/**
+ * Applies an additive time movement with various options on supported float properties
+ */
+UCLASS(AutoExpandCategories=("Animator"))
+class PROPERTYANIMATOR_API UPropertyAnimatorTime : public UPropertyAnimatorFloatBase
+{
+	GENERATED_BODY()
+
+public:
+	static constexpr const TCHAR* DefaultControllerName = TEXT("Time");
+
+	UPropertyAnimatorTime();
+
+protected:
+	//~ Begin UPropertyAnimatorFloatBase
+	virtual float Evaluate(double InTimeElapsed, const FPropertyAnimatorCoreData& InPropertyData, UPropertyAnimatorFloatContext* InOptions) const override;
+	//~ End UPropertyAnimatorFloatBase
+};
\ No newline at end of file
diff --git a/Engine/Plugins/Experimental/PropertyAnimator/Source/PropertyAnimator/Public/Animators/PropertyAnimatorWiggle.h b/Engine/Plugins/Experimental/PropertyAnimator/Source/PropertyAnimator/Public/Animators/PropertyAnimatorWiggle.h
new file mode 100644
--- /dev/null
+++ b/Engine/Plugins/Experimental/PropertyAnimator/Source/PropertyAnimator/Public/Animators/PropertyAnimatorWiggle.h
@@ -0,0 +1,25 @@
+// Copyright Epic Games, Inc. All Rights Reserved.
+
+#pragma once
+
+#include "Animators/PropertyAnimatorFloatBase.h"
+#include "PropertyAnimatorWiggle.generated.h"
+
+/**
+ * Applies an additive random wiggle movement with various options on supported float properties
+ */
+UCLASS(AutoExpandCategories=("Animator"))
+class PROPERTYANIMATOR_API UPropertyAnimatorWiggle : public UPropertyAnimatorFloatBase
+{
+	GENERATED_BODY()
+
+public:
+	static constexpr const TCHAR* DefaultControllerName = TEXT("Wiggle");
+
+	UPropertyAnimatorWiggle();
+
+protected:
+	//~ Begin UPropertyAnimatorFloatBase
+	virtual float Evaluate(double InTimeElapsed, const FPropertyAnimatorCoreData& InPropertyData, UPropertyAnimatorFloatContext* InOptions) const override;
+	//~ End UPropertyAnimatorFloatBase
+};
\ No newline at end of file
diff --git a/Engine/Plugins/Experimental/PropertyAnimator/Source/PropertyAnimator/Public/Properties/PropertyAnimatorFloatContext.h b/Engine/Plugins/Experimental/PropertyAnimator/Source/PropertyAnimator/Public/Properties/PropertyAnimatorFloatContext.h
new file mode 100644
--- /dev/null
+++ b/Engine/Plugins/Experimental/PropertyAnimator/Source/PropertyAnimator/Public/Properties/PropertyAnimatorFloatContext.h
@@ -0,0 +1,69 @@
+// Copyright Epic Games, Inc. All Rights Reserved.
+
+#pragma once
+
+#include "Properties/PropertyAnimatorCoreContext.h"
+#include "PropertyAnimatorFloatContext.generated.h"
+
+/** Property context used by float driving animator */
+UCLASS(BlueprintType)
+class PROPERTYANIMATOR_API UPropertyAnimatorFloatContext : public UPropertyAnimatorCoreContext
+{
+	GENERATED_BODY()
+
+public:
+	void SetMagnitude(float InMagnitude);
+	float GetMagnitude() const
+	{
+		return Magnitude;
+	}
+
+	void SetAmplitudeMin(double InAmplitude);
+	double GetAmplitudeMin() const
+	{
+		return AmplitudeMin;
+	}
+
+	void SetAmplitudeMax(double InAmplitude);
+	double GetAmplitudeMax() const
+	{
+		return AmplitudeMax;
+	}
+
+	void SetFrequency(float InFrequency);
+	float GetFrequency() const
+	{
+		return Frequency;
+	}
+
+	void SetTimeOffset(double InTimeOffset);
+	double GetTimeOffset() const
+	{
+		return TimeOffset;
+	}
+
+protected:
+	//~ Begin UPropertyAnimatorCoreContext
+	virtual void OnAnimatedPropertyLinked() override;
+	//~ End UPropertyAnimatorCoreContext
+
+	/** Magnitude of the effect */
+	UPROPERTY(EditAnywhere, BlueprintReadWrite, Interp, Category="Animator", meta=(ClampMin="0.0", ClampMax="1.0"))
+	float Magnitude = 1.f;
+
+	/** The minimum value should be remapped to that values */
+	UPROPERTY(EditAnywhere, BlueprintReadWrite, Interp, Category="Animator")
+	double AmplitudeMin = -1.f;
+
+	/** The maximum value should be remapped to that values */
+	UPROPERTY(EditAnywhere, BlueprintReadWrite, Interp, Category="Animator")
+	double AmplitudeMax = 1.f;
+
+	/** Number of repetition every seconds for the effect */
+	UPROPERTY(EditAnywhere, BlueprintReadWrite, Interp, Category="Animator", meta=(ClampMin="0"))
+	float Frequency = 1.f;
+
+	/** Time offset variation for evaluation */
+	UPROPERTY(EditAnywhere, BlueprintReadWrite, Interp, Category="Animator")
+	double TimeOffset = 0.f;
+};
diff --git a/Engine/Plugins/Experimental/PropertyAnimator/Source/PropertyAnimator/Public/Properties/PropertyAnimatorTextGroup.h b/Engine/Plugins/Experimental/PropertyAnimator/Source/PropertyAnimator/Public/Properties/PropertyAnimatorTextGroup.h
new file mode 100644
--- /dev/null
+++ b/Engine/Plugins/Experimental/PropertyAnimator/Source/PropertyAnimator/Public/Properties/PropertyAnimatorTextGroup.h
@@ -0,0 +1,33 @@
+// Copyright Epic Games, Inc. All Rights Reserved.
+
+#pragma once
+
+#include "Properties/PropertyAnimatorCoreGroupBase.h"
+#include "PropertyAnimatorTextGroup.generated.h"
+
+/** Group that handles text character range properties */
+UCLASS(BlueprintType)
+class PROPERTYANIMATOR_API UPropertyAnimatorTextGroup : public UPropertyAnimatorCoreGroupBase
+{
+	GENERATED_BODY()
+
+public:
+	UPropertyAnimatorTextGroup()
+		: UPropertyAnimatorCoreGroupBase(TEXT("Text"))
+	{}
+
+	//~ Begin UPropertyAnimatorCoreGroupBase
+	virtual void ManageProperties(const UPropertyAnimatorCoreContext* InContext, TArray<FPropertyAnimatorCoreData>& InOutProperties) override;
+	virtual bool IsPropertySupported(const UPropertyAnimatorCoreContext* InContext) const override;
+	//~ End UPropertyAnimatorCoreGroupBase
+
+protected:
+	UPROPERTY(EditInstanceOnly, BlueprintReadWrite, Interp, Category="Animator", meta=(ClampMin="0", ClampMax="1"))
+	float RangeStart = 0.f;
+
+	UPROPERTY(EditInstanceOnly, BlueprintReadWrite, Interp, Category="Animator", meta=(ClampMin="0", ClampMax="1"))
+	float RangeEnd = 100.f;
+
+	UPROPERTY(EditInstanceOnly, BlueprintReadWrite, Interp, Category="Animator")
+	float RangeOffset = 0.f;
+};
\ No newline at end of file
diff --git a/Engine/Plugins/Experimental/PropertyAnimatorCore/Config/BasePropertyAnimatorCore.ini b/Engine/Plugins/Experimental/PropertyAnimatorCore/Config/BasePropertyAnimatorCore.ini
new file mode 100644
--- /dev/null
+++ b/Engine/Plugins/Experimental/PropertyAnimatorCore/Config/BasePropertyAnimatorCore.ini
@@ -0,0 +1,20 @@
+﻿[CoreRedirects]
++ClassRedirects=(OldName="/Script/CustomPropertyControl.CustomPropertyControlResolver",NewName="/Script/PropertyAnimatorCore.PropertyAnimatorCoreResolver")
++ClassRedirects=(OldName="/Script/CustomPropertyControl.CustomPropertyControlHandlerBase",NewName="/Script/PropertyAnimatorCore.PropertyAnimatorCoreHandlerBase")
++ClassRedirects=(OldName="/Script/CustomPropertyControl.CustomPropertyControlFloatHandler",NewName="/Script/PropertyAnimatorCore.PropertyAnimatorCoreFloatHandler")
++ClassRedirects=(OldName="/Script/CustomPropertyControl.CustomPropertyControlDoubleHandler",NewName="/Script/PropertyAnimatorCore.PropertyAnimatorCoreDoubleHandler")
++ClassRedirects=(OldName="/Script/CustomPropertyControl.CustomPropertyControlControllerBase",NewName="/Script/PropertyAnimatorCore.PropertyAnimatorCoreBase")
++ClassRedirects=(OldName="/Script/CustomPropertyControl.CustomPropertyControlComponent",NewName="/Script/PropertyAnimatorCore.PropertyAnimatorCoreComponent")
++ClassRedirects=(OldName="/Script/PropertyController.PropertyControllerCustomTimeSource",NewName="/Script/PropertyAnimator.PropertyAnimatorCustomTimeSource")
++ClassRedirects=(OldName="/Script/PropertyController.PropertyControllerTimeSourceBase",NewName="/Script/PropertyAnimator.PropertyAnimatorTimeSourceBase")
++ClassRedirects=(OldName="/Script/PropertyController.PropertyControllerWorldTimeSource",NewName="/Script/PropertyAnimator.PropertyAnimatorWorldTimeSource")
++ClassRedirects=(OldName="/Script/PropertyAnimator.PropertyAnimatorCustomTimeSource",NewName="/Script/PropertyAnimatorCore.PropertyAnimatorCoreManualTimeSource")
++ClassRedirects=(OldName="/Script/PropertyAnimator.PropertyAnimatorTimeSourceBase",NewName="/Script/PropertyAnimatorCore.PropertyAnimatorCoreTimeSourceBase")
++ClassRedirects=(OldName="/Script/PropertyAnimator.PropertyAnimatorWorldTimeSource",NewName="/Script/PropertyAnimatorCore.PropertyAnimatorCoreWorldTimeSource")
+
++StructRedirects=(OldName="/Script/CustomPropertyControlEditor.CustomPropertyControlData",NewName="/Script/PropertyAnimatorCore.PropertyAnimatorCoreData")
+
++PropertyRedirects=(OldName="/Script/PropertyAnimatorCore.PropertyAnimatorCoreBase.ControllerDisplayName",NewName="/Script/PropertyAnimatorCore.PropertyAnimatorCoreBase.AnimatorDisplayName")
++PropertyRedirects=(OldName="/Script/PropertyAnimatorCore.PropertyAnimatorCoreComponent.Controllers",NewName="/Script/PropertyAnimatorCore.PropertyAnimatorCoreComponent.Animators")
++PropertyRedirects=(OldName="/Script/PropertyAnimatorCore.PropertyAnimatorCoreBase.bEnabled",NewName="/Script/PropertyAnimatorCore.PropertyAnimatorCoreBase.bAnimatorEnabled")
++PropertyRedirects=(OldName="/Script/PropertyAnimatorCore.PropertyAnimatorCoreComponent.bEnabled",NewName="/Script/PropertyAnimatorCore.PropertyAnimatorCoreComponent.bAnimatorsEnabled")
\ No newline at end of file
diff --git a/Engine/Plugins/Experimental/PropertyAnimatorCore/Source/PropertyAnimatorCore/Private/Animators/PropertyAnimatorCoreBase.cpp b/Engine/Plugins/Experimental/PropertyAnimatorCore/Source/PropertyAnimatorCore/Private/Animators/PropertyAnimatorCoreBase.cpp
new file mode 100644
--- /dev/null
+++ b/Engine/Plugins/Experimental/PropertyAnimatorCore/Source/PropertyAnimatorCore/Private/Animators/PropertyAnimatorCoreBase.cpp
@@ -0,0 +1,655 @@
+// Copyright Epic Games, Inc. All Rights Reserved.
+
+#include "Animators/PropertyAnimatorCoreBase.h"
+
+#include "Components/PropertyAnimatorCoreComponent.h"
+#include "GameFramework/Actor.h"
+#include "Properties/PropertyAnimatorCoreGroupBase.h"
+#include "Properties/PropertyAnimatorCoreResolver.h"
+#include "Properties/Handlers/PropertyAnimatorCoreHandlerBase.h"
+#include "Subsystems/PropertyAnimatorCoreSubsystem.h"
+#include "TimeSources/PropertyAnimatorCoreTimeSourceBase.h"
+
+DEFINE_LOG_CATEGORY_STATIC(LogPropertyAnimatorCoreBase, Log, All);
+
+UPropertyAnimatorCoreBase::FOnAnimatorUpdated UPropertyAnimatorCoreBase::OnAnimatorCreatedDelegate;
+UPropertyAnimatorCoreBase::FOnAnimatorUpdated UPropertyAnimatorCoreBase::OnAnimatorRemovedDelegate;
+UPropertyAnimatorCoreBase::FOnAnimatorUpdated UPropertyAnimatorCoreBase::OnAnimatorRenamedDelegate;
+UPropertyAnimatorCoreBase::FOnAnimatorPropertyUpdated UPropertyAnimatorCoreBase::OnAnimatorPropertyLinkedDelegate;
+UPropertyAnimatorCoreBase::FOnAnimatorPropertyUpdated UPropertyAnimatorCoreBase::OnAnimatorPropertyUnlinkedDelegate;
+
+UPropertyAnimatorCoreBase::UPropertyAnimatorCoreBase()
+{
+	if (const UPropertyAnimatorCoreSubsystem* AnimatorSubsystem = UPropertyAnimatorCoreSubsystem::Get())
+	{
+		// Apply default time source
+		const TArray<FName> TimeSources = AnimatorSubsystem->GetTimeSourceNames();
+		SetTimeSourceName(!TimeSources.IsEmpty() ? TimeSources[0] : NAME_None);
+	}
+}
+
+UPropertyAnimatorCoreComponent* UPropertyAnimatorCoreBase::GetAnimatorComponent() const
+{
+	return GetTypedOuter<UPropertyAnimatorCoreComponent>();
+}
+
+float UPropertyAnimatorCoreBase::GetAnimatorComponentMagnitude() const
+{
+	if (const UPropertyAnimatorCoreComponent* AnimatorComponent = GetAnimatorComponent())
+	{
+		return AnimatorComponent->GetAnimatorsMagnitude();
+	}
+
+	return 0.f;
+}
+
+void UPropertyAnimatorCoreBase::UpdateAnimatorDisplayName()
+{
+	TArray<FString> PropertiesNames;
+	for (const FPropertyAnimatorCoreData& LinkedProperty : GetLinkedProperties())
+	{
+		PropertiesNames.Add(LinkedProperty.GetPropertyDisplayName().ToString());
+	}
+
+	auto FindCommonPrefix = [](const TConstArrayView<FString>& InNames)->FString
+	{
+		if (InNames.IsEmpty())
+		{
+			return FString();
+		}
+
+		FString CommonPrefix = InNames[0];
+
+		for (int32 Index = 1; Index < InNames.Num(); ++Index)
+		{
+			const FString& CurrentString = InNames[Index];
+
+			int32 CommonChars = 0;
+			while (CommonChars < CommonPrefix.Len()
+				&& CommonChars < CurrentString.Len()
+				&& CommonPrefix[CommonChars] == CurrentString[CommonChars])
+			{
+				++CommonChars;
+			}
+
+			CommonPrefix = CommonPrefix.Left(CommonChars);
+		}
+
+		return CommonPrefix;
+	};
+
+	FString CommonPrefix = FindCommonPrefix(PropertiesNames);
+	CommonPrefix = CommonPrefix.TrimChar(*TEXT("."));
+
+	if (CommonPrefix.IsEmpty())
+	{
+		SetAnimatorDisplayName(GetFName());
+	}
+	else
+	{
+		SetAnimatorDisplayName(FName(GetAnimatorOriginalName().ToString() + TEXT("_") + CommonPrefix));
+	}
+}
+
+UPropertyAnimatorCoreContext* UPropertyAnimatorCoreBase::GetLinkedPropertyContext(const FPropertyAnimatorCoreData& InProperty) const
+{
+	const TObjectPtr<UPropertyAnimatorCoreContext>* PropertyOptions = LinkedProperties.FindByPredicate([&InProperty](const UPropertyAnimatorCoreContext* InOptions)
+	{
+		return InOptions && InOptions->GetAnimatedProperty() == InProperty;
+	});
+
+	return PropertyOptions ? *PropertyOptions : nullptr;
+}
+
+void UPropertyAnimatorCoreBase::PostLoad()
+{
+	CleanTimeSources();
+
+	OnTimeSourceNameChanged();
+
+	Super::PostLoad();
+
+	CleanLinkedProperties();
+
+	OnAnimatorEnabledChanged();
+}
+
+#if WITH_EDITOR
+void UPropertyAnimatorCoreBase::PreEditUndo()
+{
+	Super::PreEditUndo();
+
+	constexpr bool bForceReset = true;
+	RestoreProperties(bForceReset);
+}
+
+void UPropertyAnimatorCoreBase::PostEditUndo()
+{
+	Super::PostEditUndo();
+
+	constexpr bool bForceReset = true;
+	RestoreProperties(bForceReset);
+}
+
+void UPropertyAnimatorCoreBase::PostEditChangeProperty(FPropertyChangedEvent& PropertyChangedEvent)
+{
+	UObject::PostEditChangeProperty(PropertyChangedEvent);
+
+	const FName MemberName = PropertyChangedEvent.GetMemberPropertyName();
+
+	if (MemberName == GET_MEMBER_NAME_CHECKED(UPropertyAnimatorCoreBase, bAnimatorEnabled))
+	{
+		OnAnimatorEnabledChanged();
+	}
+	else if (MemberName == GET_MEMBER_NAME_CHECKED(UPropertyAnimatorCoreBase, TimeSourceName))
+	{
+		OnTimeSourceNameChanged();
+	}
+	else if (MemberName == GET_MEMBER_NAME_CHECKED(UPropertyAnimatorCoreBase, PropertyGroups))
+	{
+		OnPropertyGroupsChanged();
+	}
+}
+#endif
+
+AActor* UPropertyAnimatorCoreBase::GetAnimatorActor() const
+{
+	return GetTypedOuter<AActor>();
+}
+
+void UPropertyAnimatorCoreBase::SetAnimatorEnabled(bool bInIsEnabled)
+{
+	if (bAnimatorEnabled == bInIsEnabled)
+	{
+		return;
+	}
+
+	bAnimatorEnabled = bInIsEnabled;
+	OnAnimatorEnabledChanged();
+}
+
+void UPropertyAnimatorCoreBase::SetTimeSourceName(FName InTimeSourceName)
+{
+	if (TimeSourceName == InTimeSourceName)
+	{
+		return;
+	}
+
+	const TArray<FString> TimeSources = GetTimeSourceNames();
+	if (!TimeSources.Contains(InTimeSourceName))
+	{
+		return;
+	}
+
+	TimeSourceName = InTimeSourceName;
+	OnTimeSourceNameChanged();
+}
+
+FName UPropertyAnimatorCoreBase::GetAnimatorOriginalName() const
+{
+	const UPropertyAnimatorCoreBase* CDO = GetClass()->GetDefaultObject<UPropertyAnimatorCoreBase>();
+	return CDO ? CDO->AnimatorDisplayName : NAME_None;
+}
+
+bool UPropertyAnimatorCoreBase::GetPropertiesSupported(const FPropertyAnimatorCoreData& InPropertyData, TSet<FPropertyAnimatorCoreData>& OutProperties, bool bInRecursiveSearch) const
+{
+	const FProperty* LeafProperty = InPropertyData.GetLeafProperty();
+	UObject* Owner = InPropertyData.GetOwner();
+
+	// Is property editable
+	if (!LeafProperty->HasAnyPropertyFlags(EPropertyFlags::CPF_Edit))
+	{
+		return false;
+	}
+
+	// We can directly control the member property
+	if (IsPropertySupported(InPropertyData))
+	{
+		OutProperties.Add(InPropertyData);
+
+		if (bInRecursiveSearch)
+		{
+			return true;
+		}
+	}
+
+	// Look for inner properties that can be controlled too
+	TFunction<bool(TArray<FProperty*>&, UObject*, TSet<FPropertyAnimatorCoreData>&)> FindSupportedPropertiesRecursively = [this, &FindSupportedPropertiesRecursively, bInRecursiveSearch, InPropertyData](TArray<FProperty*>& InChainProperties, UObject* InOwner, TSet<FPropertyAnimatorCoreData>& OutSupportedProperties)
+	{
+		FProperty* InLeafProperty = InChainProperties.Last();
+		if (const FStructProperty* StructProp = CastField<FStructProperty>(InLeafProperty))
+		{
+			for (FProperty* Property : TFieldRange<FProperty>(StructProp->Struct))
+			{
+				if (!Property->HasAnyPropertyFlags(EPropertyFlags::CPF_Edit))
+				{
+					continue;
+				}
+
+				// Copy over resolver if any on that property
+				FPropertyAnimatorCoreData PropertyControlData(InOwner, InChainProperties, Property, InPropertyData.GetPropertyResolverClass());
+
+				// We can directly control this property
+				if (IsPropertySupported(PropertyControlData))
+				{
+					OutSupportedProperties.Add(PropertyControlData);
+
+					if (bInRecursiveSearch)
+					{
+						continue;
+					}
+				}
+
+				// Check nested properties inside this property
+				TArray<FProperty*> NestedChainProperties(InChainProperties);
+				NestedChainProperties.Add(Property);
+				FindSupportedPropertiesRecursively(NestedChainProperties, InOwner, OutSupportedProperties);
+			}
+		}
+
+		return !OutSupportedProperties.IsEmpty();
+	};
+
+	TArray<FProperty*> ChainProperties;
+	ChainProperties.Reserve(InPropertyData.GetChainProperties().Num());
+	for (const TFieldPath<FProperty>& ChainProperty : InPropertyData.GetChainProperties())
+	{
+		ChainProperties.Add(ChainProperty.Get());
+	}
+
+	return FindSupportedPropertiesRecursively(ChainProperties, Owner, OutProperties);
+}
+
+bool UPropertyAnimatorCoreBase::IsPropertySupported(const FPropertyAnimatorCoreData& InPropertyData) const
+{
+	// Without any handler we can't control the property type
+	if (!InPropertyData.GetPropertyHandler())
+	{
+		return false;
+	}
+
+	return IsPropertyDirectlySupported(InPropertyData)
+		|| IsPropertyIndirectlySupported(InPropertyData);
+}
+
+void UPropertyAnimatorCoreBase::OnAnimatorEnabled()
+{
+	UE_LOG(LogPropertyAnimatorCoreBase
+		, Log
+		, TEXT("%s : PropertyAnimator %s (%s) enabled")
+		, GetAnimatorActor() ? *GetAnimatorActor()->GetActorNameOrLabel() : TEXT("Invalid Actor")
+		, *GetAnimatorDisplayName()
+		, *GetAnimatorOriginalName().ToString());
+}
+
+void UPropertyAnimatorCoreBase::OnAnimatorDisabled()
+{
+	UE_LOG(LogPropertyAnimatorCoreBase
+		, Log
+		, TEXT("%s : PropertyAnimator %s (%s) disabled")
+		, GetAnimatorActor() ? *GetAnimatorActor()->GetActorNameOrLabel() : TEXT("Invalid Actor")
+		, *GetAnimatorDisplayName()
+		, *GetAnimatorOriginalName().ToString());
+
+	constexpr bool bForceReset = true;
+	RestoreProperties(bForceReset);
+}
+
+TSubclassOf<UPropertyAnimatorCoreContext> UPropertyAnimatorCoreBase::GetPropertyContextClass(const FPropertyAnimatorCoreData& InProperty)
+{
+	return UPropertyAnimatorCoreContext::StaticClass();
+}
+
+void UPropertyAnimatorCoreBase::OnAnimatorEnabledChanged()
+{
+	if (bAnimatorEnabled)
+	{
+		OnAnimatorEnabled();
+	}
+	else
+	{
+		OnAnimatorDisabled();
+	}
+}
+
+void UPropertyAnimatorCoreBase::CleanLinkedProperties()
+{
+	for (TArray<TObjectPtr<UPropertyAnimatorCoreContext>>::TIterator It(LinkedProperties); It; ++It)
+	{
+		const UPropertyAnimatorCoreContext* PropertyContext = It->Get();
+		if (!PropertyContext || !PropertyContext->GetAnimatedProperty().IsResolved())
+		{
+			It.RemoveCurrent();
+		}
+	}
+}
+
+void UPropertyAnimatorCoreBase::CleanTimeSources()
+{
+	for (TMap<FName, TObjectPtr<UPropertyAnimatorCoreTimeSourceBase>>::TIterator It(TimeSourcesInstances); It; ++It)
+	{
+		// Remove invalid time sources
+		if (!It->Value.Get())
+		{
+			It.RemoveCurrent();
+		}
+	}
+}
+
+void UPropertyAnimatorCoreBase::OnTimeSourceNameChanged()
+{
+	if (ActiveTimeSource)
+	{
+		ActiveTimeSource->DeactivateTimeSource();
+	}
+
+	ActiveTimeSource = FindOrAddTimeSource(TimeSourceName);
+
+	if (ActiveTimeSource)
+	{
+		ActiveTimeSource->ActivateTimeSource();
+	}
+}
+
+void UPropertyAnimatorCoreBase::EvaluateAnimator()
+{
+	UPropertyAnimatorCoreTimeSourceBase* TimeSource = GetActiveTimeSource();
+
+	if (!GetAnimatorEnabled()
+		|| !TimeSource)
+	{
+		return;
+	}
+
+	const TOptional<double> TimeElapsed = TimeSource->GetConditionalTimeElapsed();
+
+	if (!TimeElapsed.IsSet())
+	{
+		return;
+	}
+
+	RestoreProperties();
+
+	SaveProperties();
+
+	EvaluatedPropertyValues.Reset();
+
+	FPropertyAnimatorCoreEvaluationParameters Parameters;
+	Parameters.AnimatorsMagnitude = GetAnimatorComponentMagnitude();
+	Parameters.TimeElapsed = TimeElapsed.GetValue();
+
+	bEvaluatingProperties = true;
+	EvaluateProperties(Parameters);
+	bEvaluatingProperties = false;
+}
+
+void UPropertyAnimatorCoreBase::OnPropertyGroupsChanged()
+{
+	TSet<FName> CurrentGroupNames;
+
+	Algo::TransformIf(
+		PropertyGroups
+		, CurrentGroupNames
+		, [](const UPropertyAnimatorCoreGroupBase* InGroup)->bool
+		{
+			return !!InGroup;
+		}
+		, [](const UPropertyAnimatorCoreGroupBase* InGroup)->FName
+		{
+			return InGroup->GetFName();
+		}
+	);
+
+	// Remove assigned group in property context if removed from animator
+	constexpr bool bResolve = false;
+	ForEachLinkedProperty<UPropertyAnimatorCoreContext>([&CurrentGroupNames](UPropertyAnimatorCoreContext* InContext, const FPropertyAnimatorCoreData& InProperty)->bool
+	{
+		if (!CurrentGroupNames.Contains(InContext->GroupName))
+		{
+			InContext->SetGroup(nullptr);
+		}
+
+		return true;
+	}, bResolve);
+}
+
+void UPropertyAnimatorCoreBase::RestoreProperties(bool bInForce)
+{
+	constexpr bool bResolve = false;
+	ForEachLinkedProperty<UPropertyAnimatorCoreContext>([bInForce](UPropertyAnimatorCoreContext* InOptions, const FPropertyAnimatorCoreData& InPropertyData)
+	{
+		bool bRestore = false;
+
+		if (bInForce)
+		{
+			bRestore = true;
+		}
+		if (InOptions->Mode == EPropertyAnimatorCoreMode::Additive)
+		{
+			bRestore = true;
+		}
+		else if (InOptions->Mode == EPropertyAnimatorCoreMode::Absolute
+			&& (InOptions->IsResolvable() || InOptions->IsConverted()))
+		{
+			bRestore = true;
+		}
+
+		if (bRestore)
+		{
+			InOptions->Restore();
+		}
+
+		return true;
+	}, bResolve);
+}
+
+void UPropertyAnimatorCoreBase::SaveProperties()
+{
+	constexpr bool bResolve = false;
+	ForEachLinkedProperty<UPropertyAnimatorCoreContext>([](UPropertyAnimatorCoreContext* InOptions, const FPropertyAnimatorCoreData& InPropertyData)
+	{
+		InOptions->Save();
+		return true;
+	}, bResolve);
+}
+
+TArray<FString> UPropertyAnimatorCoreBase::GetTimeSourceNames() const
+{
+	TArray<FString> TimeSourceNames;
+
+	if (const UPropertyAnimatorCoreSubsystem* ControllerSubsystem = UPropertyAnimatorCoreSubsystem::Get())
+	{
+		const TArray<FName> TimeSources = ControllerSubsystem->GetTimeSourceNames();
+
+		Algo::Transform(TimeSources, TimeSourceNames, [](const FName& InTimeSource)
+		{
+			return InTimeSource.ToString();
+		});
+	}
+
+	return TimeSourceNames;
+}
+
+UPropertyAnimatorCoreTimeSourceBase* UPropertyAnimatorCoreBase::FindOrAddTimeSource(FName InTimeSourceName)
+{
+	if (IsTemplate())
+	{
+		return nullptr;
+	}
+
+	UPropertyAnimatorCoreSubsystem* Subsystem = UPropertyAnimatorCoreSubsystem::Get();
+
+	if (!Subsystem || InTimeSourceName.IsNone())
+	{
+		return nullptr;
+	}
+
+	// Check cached time source instances
+	UPropertyAnimatorCoreTimeSourceBase* NewTimeSource = nullptr;
+	if (TObjectPtr<UPropertyAnimatorCoreTimeSourceBase> const* TimeSourceInstance = TimeSourcesInstances.Find(InTimeSourceName))
+	{
+		if (const TObjectPtr<UPropertyAnimatorCoreTimeSourceBase> CachedTimeSource = *TimeSourceInstance)
+		{
+			NewTimeSource = CachedTimeSource.Get();
+		}
+	}
+
+	// Create new time source instance and cache it
+	if (!NewTimeSource)
+	{
+		NewTimeSource = Subsystem->CreateNewTimeSource(InTimeSourceName, this);
+		if (NewTimeSource)
+		{
+			TimeSourcesInstances.Add(InTimeSourceName, NewTimeSource);
+		}
+	}
+
+	return NewTimeSource;
+}
+
+void UPropertyAnimatorCoreBase::SetAnimatorDisplayName(FName InName)
+{
+	if (AnimatorDisplayName == InName)
+	{
+		return;
+	}
+
+	AnimatorDisplayName = InName;
+	OnAnimatorDisplayNameChanged();
+
+	OnAnimatorRenamedDelegate.Broadcast(this);
+}
+
+TSet<FPropertyAnimatorCoreData> UPropertyAnimatorCoreBase::GetLinkedProperties() const
+{
+	TSet<FPropertyAnimatorCoreData> LinkedPropertiesSet;
+
+	for (const TObjectPtr<UPropertyAnimatorCoreContext>& Options : LinkedProperties)
+	{
+		if (Options)
+		{
+			LinkedPropertiesSet.Emplace(Options->GetAnimatedProperty());
+		}
+	}
+
+	return LinkedPropertiesSet;
+}
+
+int32 UPropertyAnimatorCoreBase::GetLinkedPropertiesCount() const
+{
+	return LinkedProperties.Num();
+}
+
+bool UPropertyAnimatorCoreBase::LinkProperty(const FPropertyAnimatorCoreData& InLinkProperty)
+{
+	if (!InLinkProperty.IsResolved())
+	{
+		return false;
+	}
+
+	const UObject* Owner = InLinkProperty.GetOwner();
+	const AActor* OwningActor = GetTypedOuter<AActor>();
+
+	if (Owner != OwningActor && !Owner->IsIn(OwningActor))
+	{
+		return false;
+	}
+
+	if (!IsPropertySupported(InLinkProperty))
+	{
+		return false;
+	}
+
+	if (IsPropertyLinked(InLinkProperty))
+	{
+		return false;
+	}
+
+	const TSubclassOf<UPropertyAnimatorCoreContext> ContextSubclass = GetPropertyContextClass(InLinkProperty);
+	const UClass* ContextClass = ContextSubclass.Get();
+
+	if (!IsValid(ContextClass))
+	{
+		return false;
+	}
+
+	UPropertyAnimatorCoreContext* PropertyContext = NewObject<UPropertyAnimatorCoreContext>(this, ContextClass, NAME_None, RF_Transactional);
+	PropertyContext->ConstructInternal(InLinkProperty);
+
+	LinkedProperties.Add(PropertyContext);
+	OnPropertyLinked(PropertyContext);
+
+	UPropertyAnimatorCoreBase::OnAnimatorPropertyLinkedDelegate.Broadcast(this, InLinkProperty);
+
+	return true;
+}
+
+bool UPropertyAnimatorCoreBase::UnlinkProperty(const FPropertyAnimatorCoreData& InUnlinkProperty)
+{
+	if (!InUnlinkProperty.IsResolved())
+	{
+		return false;
+	}
+
+	if (!IsPropertyLinked(InUnlinkProperty))
+	{
+		return false;
+	}
+
+	const UObject* Owner = InUnlinkProperty.GetOwner();
+	const AActor* OwningActor = GetTypedOuter<AActor>();
+
+	if (Owner != OwningActor && !Owner->IsIn(OwningActor))
+	{
+		return false;
+	}
+
+	UPropertyAnimatorCoreContext* PropertyContext = GetLinkedPropertyContext(InUnlinkProperty);
+
+	PropertyContext->Restore();
+	LinkedProperties.Remove(PropertyContext);
+	OnPropertyUnlinked(PropertyContext);
+
+	UPropertyAnimatorCoreBase::OnAnimatorPropertyUnlinkedDelegate.Broadcast(this, InUnlinkProperty);
+
+	return true;
+}
+
+bool UPropertyAnimatorCoreBase::IsPropertyLinked(const FPropertyAnimatorCoreData& InPropertyData) const
+{
+	return LinkedProperties.ContainsByPredicate([&InPropertyData](const UPropertyAnimatorCoreContext* InOptions)
+	{
+		return InOptions && InOptions->GetAnimatedProperty() == InPropertyData;
+	});
+}
+
+bool UPropertyAnimatorCoreBase::IsPropertiesLinked(const TSet<FPropertyAnimatorCoreData>& InProperties) const
+{
+	const TSet<FPropertyAnimatorCoreData> LinkedPropertiesSet = GetLinkedProperties();
+	return LinkedPropertiesSet.Includes(InProperties);
+}
+
+TSet<FPropertyAnimatorCoreData> UPropertyAnimatorCoreBase::GetInnerPropertiesLinked(const FPropertyAnimatorCoreData& InPropertyData) const
+{
+	TSet<FPropertyAnimatorCoreData> OutProperties;
+
+	if (!InPropertyData.IsResolved())
+	{
+		return OutProperties;
+	}
+
+	const TFieldPath<FProperty> LeafProperty = InPropertyData.GetLeafProperty();
+
+	for (const FPropertyAnimatorCoreData& ControllerProperty : GetLinkedProperties())
+	{
+		const int32 LeafPropertyIdx = ControllerProperty.GetChainProperties().Find(LeafProperty);
+
+		// If member property is inside array and not the last one, then this controlled property is inside the InPropertyData
+		if (LeafPropertyIdx != INDEX_NONE)
+		{
+			OutProperties.Add(ControllerProperty);
+		}
+	}
+
+	return OutProperties;
+}
