[PropertyAnimator] Moved plugin

diff --git a/Engine/Plugins/Experimental/PropertyAnimatorCore/Source/PropertyAnimatorCore/Private/Components/PropertyAnimatorCoreComponent.cpp b/Engine/Plugins/Experimental/PropertyAnimatorCore/Source/PropertyAnimatorCore/Private/Components/PropertyAnimatorCoreComponent.cpp
new file mode 100644
--- /dev/null
+++ b/Engine/Plugins/Experimental/PropertyAnimatorCore/Source/PropertyAnimatorCore/Private/Components/PropertyAnimatorCoreComponent.cpp
@@ -0,0 +1,321 @@
+// Copyright Epic Games, Inc. All Rights Reserved.
+
+#include "Components/PropertyAnimatorCoreComponent.h"
+
+#include "Engine/World.h"
+#include "GameFramework/Actor.h"
+#include "Subsystems/PropertyAnimatorCoreSubsystem.h"
+
+UPropertyAnimatorCoreBase* UPropertyAnimatorCoreComponent::AddAnimator(const UClass* InAnimatorClass)
+{
+	if (!InAnimatorClass)
+	{
+		return nullptr;
+	}
+
+	UPropertyAnimatorCoreBase* NewAnimator = NewObject<UPropertyAnimatorCoreBase>(this, InAnimatorClass, NAME_None, RF_Transactional);
+
+	if (NewAnimator)
+	{
+		AnimatorsInternal = Animators;
+		Animators.Add(NewAnimator);
+
+		OnAnimatorsChanged();
+	}
+
+	return NewAnimator;
+}
+
+bool UPropertyAnimatorCoreComponent::RemoveAnimator(UPropertyAnimatorCoreBase* InAnimator)
+{
+	if (!Animators.Contains(InAnimator))
+	{
+		return false;
+	}
+
+	AnimatorsInternal = Animators;
+	Animators.Remove(InAnimator);
+
+	OnAnimatorsChanged();
+
+	return true;
+}
+
+void UPropertyAnimatorCoreComponent::OnAnimatorsSetEnabled(const UWorld* InWorld, bool bInEnabled, bool bInTransact)
+{
+	if (GetWorld() == InWorld)
+	{
+#if WITH_EDITOR
+		if (bInTransact)
+		{
+			Modify();
+		}
+#endif
+
+		SetAnimatorsEnabled(bInEnabled);
+	}
+}
+
+void UPropertyAnimatorCoreComponent::OnAnimatorsChanged()
+{
+	TSet<TObjectPtr<UPropertyAnimatorCoreBase>> RemovedAnimators = AnimatorsInternal.Difference(Animators);
+	TSet<TObjectPtr<UPropertyAnimatorCoreBase>> AddedAnimators = Animators.Difference(AnimatorsInternal);
+	AnimatorsInternal.Empty();
+
+	for (const TObjectPtr<UPropertyAnimatorCoreBase>& RemovedAnimator : RemovedAnimators)
+	{
+		if (RemovedAnimator)
+		{
+			RemovedAnimator->SetAnimatorEnabled(false);
+			RemovedAnimator->OnAnimatorRemoved();
+		}
+	}
+
+	const bool bShouldAnimatorsTick = ShouldAnimatorsTick();
+	for (const TObjectPtr<UPropertyAnimatorCoreBase>& AddedAnimator : AddedAnimators)
+	{
+		if (AddedAnimator)
+		{
+			AddedAnimator->SetAnimatorDisplayName(GetAnimatorName(AddedAnimator));
+			AddedAnimator->OnAnimatorAdded();
+			AddedAnimator->SetAnimatorEnabled(true);
+
+			if (!bShouldAnimatorsTick)
+			{
+				AddedAnimator->SetAnimatorEnabled(false);
+			}
+		}
+	}
+
+	OnAnimatorsEnabledChanged();
+}
+
+void UPropertyAnimatorCoreComponent::OnAnimatorsEnabledChanged()
+{
+	const bool bEnableAnimators = ShouldAnimatorsTick();
+
+	if (IsComponentTickEnabled() == bEnableAnimators)
+	{
+		return;
+	}
+
+	for (const TObjectPtr<UPropertyAnimatorCoreBase>& Animator : Animators)
+	{
+		if (!IsValid(Animator))
+		{
+			continue;
+		}
+
+		if (bEnableAnimators)
+		{
+			Animator->OnAnimatorEnabled();
+		}
+		else
+		{
+			Animator->OnAnimatorDisabled();
+		}
+	}
+
+	SetComponentTickEnabled(bEnableAnimators);
+}
+
+bool UPropertyAnimatorCoreComponent::ShouldAnimatorsTick() const
+{
+	return bAnimatorsEnabled
+		&& !Animators.IsEmpty()
+		&& !FMath::IsNearlyZero(AnimatorsMagnitude);
+}
+
+FName UPropertyAnimatorCoreComponent::GetAnimatorName(const UPropertyAnimatorCoreBase* InAnimator)
+{
+	if (!InAnimator)
+	{
+		return NAME_None;
+	}
+
+	FString NewAnimatorName = InAnimator->GetName();
+
+	const int32 Idx = NewAnimatorName.Find(InAnimator->GetAnimatorOriginalName().ToString());
+	if (Idx != INDEX_NONE)
+	{
+		NewAnimatorName = NewAnimatorName.RightChop(Idx);
+	}
+
+	return FName(NewAnimatorName);
+}
+
+void UPropertyAnimatorCoreComponent::TickComponent(float InDeltaTime, ELevelTick InTickType, FActorComponentTickFunction* InThisTickFunction)
+{
+	Super::TickComponent(InDeltaTime, InTickType, InThisTickFunction);
+
+	if (!ShouldAnimatorsTick())
+	{
+		return;
+	}
+
+	const UWorld* World = GetWorld();
+	const bool bIsSupportedWorld = IsValid(World) && (World->IsGameWorld() || World->IsEditorWorld());
+
+	if (!bIsSupportedWorld)
+	{
+		return;
+	}
+
+	for (const TObjectPtr<UPropertyAnimatorCoreBase>& Animator : Animators)
+	{
+		if (!IsValid(Animator))
+		{
+			continue;
+		}
+
+		Animator->EvaluateAnimator();
+	}
+}
+
+UPropertyAnimatorCoreComponent* UPropertyAnimatorCoreComponent::FindOrAdd(AActor* InActor)
+{
+	if (!IsValid(InActor))
+	{
+		return nullptr;
+	}
+
+	if (UPropertyAnimatorCoreComponent* ExistingComponent = InActor->FindComponentByClass<UPropertyAnimatorCoreComponent>())
+	{
+		return ExistingComponent;
+	}
+
+#if WITH_EDITOR
+	InActor->Modify();
+#endif
+
+	const UClass* const ComponentClass = UPropertyAnimatorCoreComponent::StaticClass();
+
+	// Construct the new component and attach as needed
+	UPropertyAnimatorCoreComponent* const PropertyAnimatorComponent = NewObject<UPropertyAnimatorCoreComponent>(InActor
+		, ComponentClass
+		, MakeUniqueObjectName(InActor, ComponentClass, TEXT("PropertyAnimatorComponent"))
+		, RF_Transactional);
+
+	// Add to SerializedComponents array so it gets saved
+	InActor->AddInstanceComponent(PropertyAnimatorComponent);
+	PropertyAnimatorComponent->OnComponentCreated();
+	PropertyAnimatorComponent->RegisterComponent();
+
+#if WITH_EDITOR
+	// Rerun construction scripts
+	InActor->RerunConstructionScripts();
+#endif
+
+	return PropertyAnimatorComponent;
+}
+
+UPropertyAnimatorCoreComponent::UPropertyAnimatorCoreComponent()
+{
+	if (!IsTemplate())
+	{
+		bTickInEditor = true;
+		PrimaryComponentTick.bCanEverTick = true;
+		PrimaryComponentTick.bHighPriority = true;
+
+		// Used to toggle animators state in world
+		UPropertyAnimatorCoreSubsystem::OnAnimatorsSetEnabledDelegate.AddUObject(this, &UPropertyAnimatorCoreComponent::OnAnimatorsSetEnabled);
+	}
+}
+
+void UPropertyAnimatorCoreComponent::SetAnimatorsEnabled(bool bInEnabled)
+{
+	if (bAnimatorsEnabled == bInEnabled)
+	{
+		return;
+	}
+
+	bAnimatorsEnabled = bInEnabled;
+	OnAnimatorsEnabledChanged();
+}
+
+void UPropertyAnimatorCoreComponent::SetAnimatorsMagnitude(float InMagnitude)
+{
+	InMagnitude = FMath::Clamp(InMagnitude, 0.f, 1.f);
+
+	if (FMath::IsNearlyEqual(AnimatorsMagnitude, InMagnitude))
+	{
+		return;
+	}
+
+	AnimatorsMagnitude = InMagnitude;
+	OnAnimatorsEnabledChanged();
+}
+
+void UPropertyAnimatorCoreComponent::DestroyComponent(bool bPromoteChildren)
+{
+	AnimatorsInternal = Animators;
+	Animators.Empty();
+
+	OnAnimatorsChanged();
+
+	Super::DestroyComponent(bPromoteChildren);
+}
+
+#if WITH_EDITOR
+void UPropertyAnimatorCoreComponent::PostEditUndo()
+{
+	Super::PostEditUndo();
+
+	OnAnimatorsEnabledChanged();
+}
+
+void UPropertyAnimatorCoreComponent::PreEditChange(FProperty* PropertyAboutToChange)
+{
+	Super::PreEditChange(PropertyAboutToChange);
+
+	const FName MemberName = PropertyAboutToChange ? PropertyAboutToChange->GetFName() : NAME_None;
+
+	if (MemberName == GET_MEMBER_NAME_CHECKED(UPropertyAnimatorCoreComponent, Animators))
+	{
+		AnimatorsInternal = Animators;
+	}
+}
+
+void UPropertyAnimatorCoreComponent::PostEditChangeProperty(FPropertyChangedEvent& PropertyChangedEvent)
+{
+	Super::PostEditChangeProperty(PropertyChangedEvent);
+
+	const FName MemberName = PropertyChangedEvent.GetMemberPropertyName();
+
+	if (MemberName == GET_MEMBER_NAME_CHECKED(UPropertyAnimatorCoreComponent, bAnimatorsEnabled)
+		|| MemberName == GET_MEMBER_NAME_CHECKED(UPropertyAnimatorCoreComponent, AnimatorsMagnitude))
+	{
+		OnAnimatorsEnabledChanged();
+	}
+	else if (MemberName == GET_MEMBER_NAME_CHECKED(UPropertyAnimatorCoreComponent, Animators))
+	{
+		OnAnimatorsChanged();
+	}
+}
+#endif
+
+void UPropertyAnimatorCoreComponent::SetAnimators(const TSet<TObjectPtr<UPropertyAnimatorCoreBase>>& InAnimators)
+{
+	if (Animators.Includes(InAnimators) && Animators.Num() == InAnimators.Num())
+	{
+		return;
+	}
+
+	AnimatorsInternal = Animators;
+	Animators = InAnimators;
+	OnAnimatorsChanged();
+}
+
+void UPropertyAnimatorCoreComponent::ForEachAnimator(const TFunctionRef<bool(UPropertyAnimatorCoreBase*)> InFunction) const
+{
+	for (const TObjectPtr<UPropertyAnimatorCoreBase>& Animator : Animators)
+	{
+		if (Animator)
+		{
+			if (!InFunction(Animator))
+			{
+				break;
+			}
+		}
+	}
+}
diff --git a/Engine/Plugins/Experimental/PropertyAnimatorCore/Source/PropertyAnimatorCore/Private/Presets/PropertyAnimatorCorePresetBase.cpp b/Engine/Plugins/Experimental/PropertyAnimatorCore/Source/PropertyAnimatorCore/Private/Presets/PropertyAnimatorCorePresetBase.cpp
new file mode 100644
--- /dev/null
+++ b/Engine/Plugins/Experimental/PropertyAnimatorCore/Source/PropertyAnimatorCore/Private/Presets/PropertyAnimatorCorePresetBase.cpp
@@ -0,0 +1,107 @@
+// Copyright Epic Games, Inc. All Rights Reserved.
+
+#include "Presets/PropertyAnimatorCorePresetBase.h"
+
+#include "Animators/PropertyAnimatorCoreBase.h"
+
+FString UPropertyAnimatorCorePresetBase::GetPresetDisplayName() const
+{
+	return FName::NameToDisplayString(PresetName.ToString(), false);
+}
+
+void UPropertyAnimatorCorePresetBase::GetSupportedPresetProperties(const AActor* InActor, const UPropertyAnimatorCoreBase* InAnimator, TSet<FPropertyAnimatorCoreData>& OutProperties) const
+{
+	TSet<FPropertyAnimatorCoreData> PresetProperties;
+	GetPresetProperties(InActor, InAnimator, PresetProperties);
+
+	if (PresetProperties.IsEmpty())
+	{
+		return;
+	}
+
+	for (const FPropertyAnimatorCoreData& PresetProperty : PresetProperties)
+	{
+		InAnimator->GetPropertiesSupported(PresetProperty, OutProperties, false);
+	}
+}
+
+bool UPropertyAnimatorCorePresetBase::IsPresetSupported(const AActor* InActor, const UPropertyAnimatorCoreBase* InAnimator) const
+{
+	if (!IsValid(InActor) || !IsValid(InAnimator))
+	{
+		return false;
+	}
+
+	TSet<FPropertyAnimatorCoreData> SupportedProperties;
+	GetSupportedPresetProperties(InActor, InAnimator, SupportedProperties);
+
+	return !SupportedProperties.IsEmpty();
+}
+
+bool UPropertyAnimatorCorePresetBase::ApplyPreset(UPropertyAnimatorCoreBase* InAnimator)
+{
+	if (!IsValid(InAnimator) || InAnimator->IsTemplate())
+	{
+		return false;
+	}
+
+	TSet<FPropertyAnimatorCoreData> SupportedProperties;
+	GetSupportedPresetProperties(InAnimator->GetAnimatorActor(), InAnimator, SupportedProperties);
+
+	if (SupportedProperties.IsEmpty())
+	{
+		return false;
+	}
+
+	for (FPropertyAnimatorCoreData& SupportedProperty : SupportedProperties)
+	{
+		InAnimator->LinkProperty(SupportedProperty);
+	}
+
+	InAnimator->SetAnimatorDisplayName(FName(InAnimator->GetAnimatorOriginalName().ToString() + TEXT("_") + GetPresetDisplayName()));
+
+	OnPresetApplied(InAnimator, SupportedProperties);
+
+	return true;
+}
+
+bool UPropertyAnimatorCorePresetBase::IsPresetApplied(const UPropertyAnimatorCoreBase* InAnimator) const
+{
+	if (!IsValid(InAnimator) || InAnimator->IsTemplate())
+	{
+		return false;
+	}
+
+	TSet<FPropertyAnimatorCoreData> SupportedProperties;
+	GetSupportedPresetProperties(InAnimator->GetAnimatorActor(), InAnimator, SupportedProperties);
+
+	if (SupportedProperties.IsEmpty())
+	{
+		return false;
+	}
+
+	return InAnimator->IsPropertiesLinked(SupportedProperties);
+}
+
+bool UPropertyAnimatorCorePresetBase::UnapplyPreset(UPropertyAnimatorCoreBase* InAnimator)
+{
+	if (!IsValid(InAnimator) || InAnimator->IsTemplate())
+	{
+		return false;
+	}
+
+	TSet<FPropertyAnimatorCoreData> SupportedProperties;
+	GetSupportedPresetProperties(InAnimator->GetAnimatorActor(), InAnimator, SupportedProperties);
+
+	if (SupportedProperties.IsEmpty())
+	{
+		return false;
+	}
+
+	for (FPropertyAnimatorCoreData& SupportedProperty : SupportedProperties)
+	{
+		InAnimator->UnlinkProperty(SupportedProperty);
+	}
+
+	return true;
+}
diff --git a/Engine/Plugins/Experimental/PropertyAnimatorCore/Source/PropertyAnimatorCore/Private/Properties/Handlers/PropertyAnimatorCoreBoolHandler.cpp b/Engine/Plugins/Experimental/PropertyAnimatorCore/Source/PropertyAnimatorCore/Private/Properties/Handlers/PropertyAnimatorCoreBoolHandler.cpp
new file mode 100644
--- /dev/null
+++ b/Engine/Plugins/Experimental/PropertyAnimatorCore/Source/PropertyAnimatorCore/Private/Properties/Handlers/PropertyAnimatorCoreBoolHandler.cpp
@@ -0,0 +1,41 @@
+// Copyright Epic Games, Inc. All Rights Reserved.
+
+#include "Properties/Handlers/PropertyAnimatorCoreBoolHandler.h"
+
+bool UPropertyAnimatorCoreBoolHandler::IsPropertySupported(const FPropertyAnimatorCoreData& InPropertyData) const
+{
+	if (InPropertyData.IsA<FBoolProperty>())
+	{
+		return true;
+	}
+
+	return Super::IsPropertySupported(InPropertyData);
+}
+
+bool UPropertyAnimatorCoreBoolHandler::GetValue(const FPropertyAnimatorCoreData& InPropertyData, FInstancedPropertyBag& OutValue)
+{
+	const FName PropertyName(InPropertyData.GetPathHash());
+	OutValue.AddProperty(PropertyName, EPropertyBagPropertyType::Bool);
+
+	bool bValue;
+	InPropertyData.GetPropertyValuePtr(&bValue);
+
+	OutValue.SetValueBool(PropertyName, bValue);
+
+	return true;
+}
+
+bool UPropertyAnimatorCoreBoolHandler::SetValue(const FPropertyAnimatorCoreData& InPropertyData, const FInstancedPropertyBag& InValue)
+{
+	const FName PropertyName(InPropertyData.GetPathHash());
+	TValueOrError<bool, EPropertyBagResult> ValueResult = InValue.GetValueBool(PropertyName);
+	if (!ValueResult.HasValue())
+	{
+		return false;
+	}
+
+	bool bNewValue = ValueResult.GetValue();
+	InPropertyData.SetPropertyValuePtr(&bNewValue);
+
+	return true;
+}
\ No newline at end of file
diff --git a/Engine/Plugins/Experimental/PropertyAnimatorCore/Source/PropertyAnimatorCore/Private/Properties/Handlers/PropertyAnimatorCoreDoubleHandler.cpp b/Engine/Plugins/Experimental/PropertyAnimatorCore/Source/PropertyAnimatorCore/Private/Properties/Handlers/PropertyAnimatorCoreDoubleHandler.cpp
new file mode 100644
--- /dev/null
+++ b/Engine/Plugins/Experimental/PropertyAnimatorCore/Source/PropertyAnimatorCore/Private/Properties/Handlers/PropertyAnimatorCoreDoubleHandler.cpp
@@ -0,0 +1,82 @@
+// Copyright Epic Games, Inc. All Rights Reserved.
+
+#include "Properties/Handlers/PropertyAnimatorCoreDoubleHandler.h"
+
+bool UPropertyAnimatorCoreDoubleHandler::IsPropertySupported(const FPropertyAnimatorCoreData& InPropertyData) const
+{
+	if (InPropertyData.IsA<FDoubleProperty>())
+	{
+		return true;
+	}
+
+	return Super::IsPropertySupported(InPropertyData);
+}
+
+bool UPropertyAnimatorCoreDoubleHandler::GetValue(const FPropertyAnimatorCoreData& InPropertyData, FInstancedPropertyBag& OutValue)
+{
+	const FName PropertyName(InPropertyData.GetPathHash());
+	OutValue.AddProperty(PropertyName, EPropertyBagPropertyType::Double);
+
+	double Value;
+	InPropertyData.GetPropertyValuePtr(&Value);
+
+	OutValue.SetValueDouble(PropertyName, Value);
+
+	return true;
+}
+
+bool UPropertyAnimatorCoreDoubleHandler::SetValue(const FPropertyAnimatorCoreData& InPropertyData, const FInstancedPropertyBag& InValue)
+{
+	const FName PropertyName(InPropertyData.GetPathHash());
+	TValueOrError<double, EPropertyBagResult> ValueResult = InValue.GetValueDouble(PropertyName);
+	if (!ValueResult.HasValue())
+	{
+		return false;
+	}
+
+	double NewValue = ValueResult.GetValue();
+	InPropertyData.SetPropertyValuePtr(&NewValue);
+
+	return true;
+}
+
+bool UPropertyAnimatorCoreDoubleHandler::IsAdditiveSupported() const
+{
+	return true;
+}
+
+bool UPropertyAnimatorCoreDoubleHandler::AddValue(const FPropertyAnimatorCoreData& InPropertyData, const FInstancedPropertyBag& InValue)
+{
+	const FName PropertyName(InPropertyData.GetPathHash());
+	const TValueOrError<double, EPropertyBagResult> ValueResult = InValue.GetValueDouble(PropertyName);
+	if (!ValueResult.HasValue())
+	{
+		return false;
+	}
+
+	double Value;
+	InPropertyData.GetPropertyValuePtr(&Value);
+
+	double NewValue = Value + ValueResult.GetValue();
+	InPropertyData.SetPropertyValuePtr(&NewValue);
+
+	return true;
+}
+
+bool UPropertyAnimatorCoreDoubleHandler::SubtractValue(const FPropertyAnimatorCoreData& InPropertyData, const FInstancedPropertyBag& InValue)
+{
+	const FName PropertyName(InPropertyData.GetPathHash());
+	const TValueOrError<double, EPropertyBagResult> ValueResult = InValue.GetValueDouble(PropertyName);
+	if (!ValueResult.HasValue())
+	{
+		return false;
+	}
+
+	double Value;
+	InPropertyData.GetPropertyValuePtr(&Value);
+
+	double NewValue = Value - ValueResult.GetValue();
+	InPropertyData.SetPropertyValuePtr(&NewValue);
+
+	return true;
+}
diff --git a/Engine/Plugins/Experimental/PropertyAnimatorCore/Source/PropertyAnimatorCore/Private/Properties/Handlers/PropertyAnimatorCoreIntHandler.cpp b/Engine/Plugins/Experimental/PropertyAnimatorCore/Source/PropertyAnimatorCore/Private/Properties/Handlers/PropertyAnimatorCoreIntHandler.cpp
new file mode 100644
--- /dev/null
+++ b/Engine/Plugins/Experimental/PropertyAnimatorCore/Source/PropertyAnimatorCore/Private/Properties/Handlers/PropertyAnimatorCoreIntHandler.cpp
@@ -0,0 +1,82 @@
+// Copyright Epic Games, Inc. All Rights Reserved.
+
+#include "Properties/Handlers/PropertyAnimatorCoreIntHandler.h"
+
+bool UPropertyAnimatorCoreIntHandler::IsPropertySupported(const FPropertyAnimatorCoreData& InPropertyData) const
+{
+	if (InPropertyData.IsA<FIntProperty>())
+	{
+		return true;
+	}
+
+	return Super::IsPropertySupported(InPropertyData);
+}
+
+bool UPropertyAnimatorCoreIntHandler::GetValue(const FPropertyAnimatorCoreData& InPropertyData, FInstancedPropertyBag& OutValue)
+{
+	const FName PropertyName(InPropertyData.GetPathHash());
+	OutValue.AddProperty(PropertyName, EPropertyBagPropertyType::Int32);
+
+	int32 Value;
+	InPropertyData.GetPropertyValuePtr(&Value);
+
+	OutValue.SetValueInt32(PropertyName, Value);
+
+	return true;
+}
+
+bool UPropertyAnimatorCoreIntHandler::SetValue(const FPropertyAnimatorCoreData& InPropertyData, const FInstancedPropertyBag& InValue)
+{
+	const FName PropertyName(InPropertyData.GetPathHash());
+	TValueOrError<int32, EPropertyBagResult> ValueResult = InValue.GetValueInt32(PropertyName);
+	if (!ValueResult.HasValue())
+	{
+		return false;
+	}
+
+	int32 NewValue = ValueResult.GetValue();
+	InPropertyData.SetPropertyValuePtr(&NewValue);
+
+	return true;
+}
+
+bool UPropertyAnimatorCoreIntHandler::IsAdditiveSupported() const
+{
+	return true;
+}
+
+bool UPropertyAnimatorCoreIntHandler::AddValue(const FPropertyAnimatorCoreData& InPropertyData, const FInstancedPropertyBag& InValue)
+{
+	const FName PropertyName(InPropertyData.GetPathHash());
+	const TValueOrError<int32, EPropertyBagResult> ValueResult = InValue.GetValueInt32(PropertyName);
+	if (!ValueResult.HasValue())
+	{
+		return false;
+	}
+
+	int32 Value;
+	InPropertyData.GetPropertyValuePtr(&Value);
+
+	int32 NewValue = Value + ValueResult.GetValue();
+	InPropertyData.SetPropertyValuePtr(&NewValue);
+
+	return true;
+}
+
+bool UPropertyAnimatorCoreIntHandler::SubtractValue(const FPropertyAnimatorCoreData& InPropertyData, const FInstancedPropertyBag& InValue)
+{
+	const FName PropertyName(InPropertyData.GetPathHash());
+	const TValueOrError<int32, EPropertyBagResult> ValueResult = InValue.GetValueInt32(PropertyName);
+	if (!ValueResult.HasValue())
+	{
+		return false;
+	}
+
+	int32 Value;
+	InPropertyData.GetPropertyValuePtr(&Value);
+
+	int32 NewValue = Value - ValueResult.GetValue();
+	InPropertyData.SetPropertyValuePtr(&NewValue);
+
+	return true;
+}
diff --git a/Engine/Plugins/Experimental/PropertyAnimatorCore/Source/PropertyAnimatorCore/Private/Properties/Handlers/PropertyAnimatorCoreStringHandler.cpp b/Engine/Plugins/Experimental/PropertyAnimatorCore/Source/PropertyAnimatorCore/Private/Properties/Handlers/PropertyAnimatorCoreStringHandler.cpp
new file mode 100644
--- /dev/null
+++ b/Engine/Plugins/Experimental/PropertyAnimatorCore/Source/PropertyAnimatorCore/Private/Properties/Handlers/PropertyAnimatorCoreStringHandler.cpp
@@ -0,0 +1,90 @@
+// Copyright Epic Games, Inc. All Rights Reserved.
+
+#include "Properties/Handlers/PropertyAnimatorCoreStringHandler.h"
+
+bool UPropertyAnimatorCoreStringHandler::IsPropertySupported(const FPropertyAnimatorCoreData& InPropertyData) const
+{
+	if (InPropertyData.IsA<FStrProperty>())
+	{
+		return true;
+	}
+
+	return Super::IsPropertySupported(InPropertyData);
+}
+
+bool UPropertyAnimatorCoreStringHandler::GetValue(const FPropertyAnimatorCoreData& InPropertyData, FInstancedPropertyBag& OutValue)
+{
+	const FName PropertyName(InPropertyData.GetPathHash());
+	OutValue.AddProperty(PropertyName, EPropertyBagPropertyType::String);
+
+	FString Value;
+	InPropertyData.GetPropertyValuePtr(&Value);
+
+	OutValue.SetValueString(PropertyName, Value);
+
+	return true;
+}
+
+bool UPropertyAnimatorCoreStringHandler::SetValue(const FPropertyAnimatorCoreData& InPropertyData, const FInstancedPropertyBag& InValue)
+{
+	const FName PropertyName(InPropertyData.GetPathHash());
+	TValueOrError<FString, EPropertyBagResult> ValueResult = InValue.GetValueString(PropertyName);
+	if (!ValueResult.HasValue())
+	{
+		return false;
+	}
+
+	FString NewValue = ValueResult.GetValue();
+	InPropertyData.SetPropertyValuePtr(&NewValue);
+
+	return true;
+}
+
+bool UPropertyAnimatorCoreStringHandler::IsAdditiveSupported() const
+{
+	return true;
+}
+
+bool UPropertyAnimatorCoreStringHandler::AddValue(const FPropertyAnimatorCoreData& InPropertyData, const FInstancedPropertyBag& InValue)
+{
+	const FName PropertyName(InPropertyData.GetPathHash());
+	const TValueOrError<FString, EPropertyBagResult> ValueResult = InValue.GetValueString(PropertyName);
+	if (!ValueResult.HasValue())
+	{
+		return false;
+	}
+
+	FString Value;
+	InPropertyData.GetPropertyValuePtr(&Value);
+
+	// Append
+	FString NewValue = Value + ValueResult.GetValue();
+	InPropertyData.SetPropertyValuePtr(&NewValue);
+
+	return true;
+}
+
+bool UPropertyAnimatorCoreStringHandler::SubtractValue(const FPropertyAnimatorCoreData& InPropertyData, const FInstancedPropertyBag& InValue)
+{
+	const FName PropertyName(InPropertyData.GetPathHash());
+	const TValueOrError<FString, EPropertyBagResult> ValueResult = InValue.GetValueString(PropertyName);
+	if (!ValueResult.HasValue())
+	{
+		return false;
+	}
+
+	FString Value;
+	InPropertyData.GetPropertyValuePtr(&Value);
+
+	// Trim end
+	FString NewValue = Value;
+	const FString TrimString = ValueResult.GetValue();
+	if (NewValue.EndsWith(TrimString))
+	{
+		NewValue.LeftChopInline(TrimString.Len());
+	}
+
+	InPropertyData.SetPropertyValuePtr(&NewValue);
+
+	return true;
+}
diff --git a/Engine/Plugins/Experimental/PropertyAnimatorCore/Source/PropertyAnimatorCore/Private/Properties/Handlers/PropertyAnimatorCoreTextHandler.cpp b/Engine/Plugins/Experimental/PropertyAnimatorCore/Source/PropertyAnimatorCore/Private/Properties/Handlers/PropertyAnimatorCoreTextHandler.cpp
new file mode 100644
--- /dev/null
+++ b/Engine/Plugins/Experimental/PropertyAnimatorCore/Source/PropertyAnimatorCore/Private/Properties/Handlers/PropertyAnimatorCoreTextHandler.cpp
@@ -0,0 +1,93 @@
+// Copyright Epic Games, Inc. All Rights Reserved.
+
+#include "Properties/Handlers/PropertyAnimatorCoreTextHandler.h"
+#include "PropertyBag.h"
+#include "UObject/TextProperty.h"
+
+bool UPropertyAnimatorCoreTextHandler::IsPropertySupported(const FPropertyAnimatorCoreData& InPropertyData) const
+{
+	if (InPropertyData.IsA<FTextProperty>())
+	{
+		return true;
+	}
+
+	return Super::IsPropertySupported(InPropertyData);
+}
+
+bool UPropertyAnimatorCoreTextHandler::GetValue(const FPropertyAnimatorCoreData& InPropertyData, FInstancedPropertyBag& OutValue)
+{
+	const FName PropertyName(InPropertyData.GetPathHash());
+	OutValue.AddProperty(PropertyName, EPropertyBagPropertyType::Text);
+
+	FText Value;
+	InPropertyData.GetPropertyValuePtr(&Value);
+
+	OutValue.SetValueText(PropertyName, Value);
+
+	return true;
+}
+
+bool UPropertyAnimatorCoreTextHandler::SetValue(const FPropertyAnimatorCoreData& InPropertyData, const FInstancedPropertyBag& InValue)
+{
+	const FName PropertyName(InPropertyData.GetPathHash());
+	TValueOrError<FText, EPropertyBagResult> ValueResult = InValue.GetValueText(PropertyName);
+	if (!ValueResult.HasValue())
+	{
+		return false;
+	}
+
+	FText NewValue = ValueResult.GetValue();
+	InPropertyData.SetPropertyValuePtr(&NewValue);
+
+	return true;
+}
+
+bool UPropertyAnimatorCoreTextHandler::IsAdditiveSupported() const
+{
+	return true;
+}
+
+bool UPropertyAnimatorCoreTextHandler::AddValue(const FPropertyAnimatorCoreData& InPropertyData, const FInstancedPropertyBag& InValue)
+{
+	const FName PropertyName(InPropertyData.GetPathHash());
+	const TValueOrError<FText, EPropertyBagResult> ValueResult = InValue.GetValueText(PropertyName);
+	if (!ValueResult.HasValue())
+	{
+		return false;
+	}
+
+	FText Value;
+	InPropertyData.GetPropertyValuePtr(&Value);
+
+	// Append
+	FText NewValue = FText::FromString(Value.ToString() + ValueResult.GetValue().ToString());
+	InPropertyData.SetPropertyValuePtr(&NewValue);
+
+	return true;
+}
+
+bool UPropertyAnimatorCoreTextHandler::SubtractValue(const FPropertyAnimatorCoreData& InPropertyData, const FInstancedPropertyBag& InValue)
+{
+	const FName PropertyName(InPropertyData.GetPathHash());
+	const TValueOrError<FText, EPropertyBagResult> ValueResult = InValue.GetValueText(PropertyName);
+	if (!ValueResult.HasValue())
+	{
+		return false;
+	}
+
+	FText Value;
+	InPropertyData.GetPropertyValuePtr(&Value);
+
+	// Trim end
+	FString StringValue = Value.ToString();
+	const FString StringValueResult = ValueResult.GetValue().ToString();
+	if (StringValue.EndsWith(StringValueResult))
+	{
+		StringValue.LeftChopInline(StringValueResult.Len());
+	}
+
+	FText NewValue = FText::FromString(StringValue);
+	InPropertyData.SetPropertyValuePtr(&NewValue);
+
+	return true;
+}
diff --git a/Engine/Plugins/Experimental/PropertyAnimatorCore/Source/PropertyAnimatorCore/Private/Properties/PropertyAnimatorCoreContext.cpp b/Engine/Plugins/Experimental/PropertyAnimatorCore/Source/PropertyAnimatorCore/Private/Properties/PropertyAnimatorCoreContext.cpp
new file mode 100644
--- /dev/null
+++ b/Engine/Plugins/Experimental/PropertyAnimatorCore/Source/PropertyAnimatorCore/Private/Properties/PropertyAnimatorCoreContext.cpp
@@ -0,0 +1,370 @@
+// Copyright Epic Games, Inc. All Rights Reserved.
+
+#include "Properties/PropertyAnimatorCoreContext.h"
+
+#include "InstancedStruct.h"
+#include "Properties/PropertyAnimatorCoreGroupBase.h"
+#include "Properties/PropertyAnimatorCoreResolver.h"
+#include "Properties/Converters/PropertyAnimatorCoreConverterBase.h"
+#include "Properties/Handlers/PropertyAnimatorCoreHandlerBase.h"
+#include "Subsystems/PropertyAnimatorCoreSubsystem.h"
+
+TArray<FPropertyAnimatorCoreData> UPropertyAnimatorCoreContext::ResolveProperty(bool bInForEvaluation) const
+{
+	TArray<FPropertyAnimatorCoreData> ResolvedProperties;
+
+	if (UPropertyAnimatorCoreResolver* Resolver = AnimatedProperty.GetPropertyResolver())
+	{
+		Resolver->ResolveProperties(AnimatedProperty, ResolvedProperties);
+	}
+	else
+	{
+		ResolvedProperties.Add(AnimatedProperty);
+	}
+
+	if (bInForEvaluation && Group)
+	{
+		Group->ManageProperties(this, ResolvedProperties);
+	}
+
+	return ResolvedProperties;
+}
+
+UPropertyAnimatorCoreBase* UPropertyAnimatorCoreContext::GetAnimator() const
+{
+	return GetTypedOuter<UPropertyAnimatorCoreBase>();
+}
+
+UPropertyAnimatorCoreHandlerBase* UPropertyAnimatorCoreContext::GetHandler() const
+{
+	if (!HandlerWeak.IsValid())
+	{
+		if (const UPropertyAnimatorCoreSubsystem* AnimatorSubsystem = UPropertyAnimatorCoreSubsystem::Get())
+		{
+			UPropertyAnimatorCoreContext* MutableThis = const_cast<UPropertyAnimatorCoreContext*>(this);
+			MutableThis->HandlerWeak = AnimatorSubsystem->GetHandler(AnimatedProperty);
+		}
+	}
+
+	return HandlerWeak.Get();
+}
+
+void UPropertyAnimatorCoreContext::SetGroup(UPropertyAnimatorCoreGroupBase* InGroup)
+{
+	if (InGroup && InGroup->IsPropertySupported(this))
+	{
+		Group = InGroup;
+		GroupName = InGroup->GetFName();
+	}
+	else
+	{
+		Group = nullptr;
+		GroupName = NAME_None;
+	}
+}
+
+bool UPropertyAnimatorCoreContext::IsResolvable() const
+{
+	return AnimatedProperty.IsResolvable();
+}
+
+bool UPropertyAnimatorCoreContext::IsConverted() const
+{
+	return !!ConverterClass.Get();
+}
+
+void UPropertyAnimatorCoreContext::SetAnimated(bool bInAnimated)
+{
+	if (bAnimated == bInAnimated)
+	{
+		return;
+	}
+
+	bAnimated = bInAnimated;
+	OnAnimatedChanged();
+}
+
+void UPropertyAnimatorCoreContext::SetMode(EPropertyAnimatorCoreMode InMode)
+{
+	if (InMode == Mode)
+	{
+		return;
+	}
+
+	Restore();
+	Mode = InMode;
+	OnModeChanged();
+}
+
+void UPropertyAnimatorCoreContext::SetConverterClass(TSubclassOf<UPropertyAnimatorCoreConverterBase> InConverterClass)
+{
+	ConverterClass = InConverterClass;
+
+	if (const UPropertyAnimatorCoreConverterBase* Converter = InConverterClass.GetDefaultObject())
+	{
+		if (UScriptStruct* RuleStruct = Converter->GetConversionRuleStruct())
+		{
+			ConverterRule = FInstancedStruct(RuleStruct);
+		}
+	}
+}
+
+void UPropertyAnimatorCoreContext::SetGroupName(FName InGroupName)
+{
+	if (GroupName.IsEqual(InGroupName))
+	{
+		return;
+	}
+
+	const TArray<FName> GroupNames = GetSupportedGroupNames();
+	if (!GroupNames.Contains(InGroupName))
+	{
+		return;
+	}
+
+	GroupName = InGroupName;
+	OnGroupNameChanged();
+}
+
+void UPropertyAnimatorCoreContext::PostLoad()
+{
+	Super::PostLoad();
+
+	CheckModeEdit();
+}
+
+#if WITH_EDITOR
+void UPropertyAnimatorCoreContext::PreEditChange(FProperty* InPropertyAboutToChange)
+{
+	Super::PreEditChange(InPropertyAboutToChange);
+
+	if (!InPropertyAboutToChange)
+	{
+		return;
+	}
+
+	const FName MemberName = InPropertyAboutToChange->GetFName();
+
+	if (MemberName == GET_MEMBER_NAME_CHECKED(UPropertyAnimatorCoreContext, Mode))
+	{
+		Restore();
+	}
+}
+
+void UPropertyAnimatorCoreContext::PostEditChangeProperty(FPropertyChangedEvent& InPropertyChangedEvent)
+{
+	Super::PostEditChangeProperty(InPropertyChangedEvent);
+
+	const FName MemberName = InPropertyChangedEvent.GetMemberPropertyName();
+
+	if (MemberName == GET_MEMBER_NAME_CHECKED(UPropertyAnimatorCoreContext, Mode))
+	{
+		OnModeChanged();
+	}
+	else if (MemberName == GET_MEMBER_NAME_CHECKED(UPropertyAnimatorCoreContext, bAnimated))
+	{
+		OnAnimatedChanged();
+	}
+	else if (MemberName == GET_MEMBER_NAME_CHECKED(UPropertyAnimatorCoreContext, GroupName))
+	{
+		OnGroupNameChanged();
+	}
+}
+#endif
+
+void UPropertyAnimatorCoreContext::OnModeChanged()
+{
+	if (const UPropertyAnimatorCoreHandlerBase* Handler = GetHandler())
+	{
+		if (Mode == EPropertyAnimatorCoreMode::Additive && !Handler->IsAdditiveSupported())
+		{
+			Mode = EPropertyAnimatorCoreMode::Absolute;
+		}
+
+		Save();
+	}
+}
+
+void UPropertyAnimatorCoreContext::OnGroupNameChanged()
+{
+	if (const UPropertyAnimatorCoreBase* Animator = GetAnimator())
+	{
+		const TObjectPtr<UPropertyAnimatorCoreGroupBase>* PropertyGroup = Animator->PropertyGroups.FindByPredicate([this](const UPropertyAnimatorCoreGroupBase* InGroup)
+		{
+			return InGroup && InGroup->GetFName() == GroupName;
+		});
+
+		SetGroup(PropertyGroup ? PropertyGroup->Get() : nullptr);
+	}
+}
+
+void UPropertyAnimatorCoreContext::ConstructInternal(const FPropertyAnimatorCoreData& InProperty)
+{
+	AnimatedProperty = InProperty;
+	CheckModeEdit();
+	OnAnimatedPropertyLinked();
+}
+
+void UPropertyAnimatorCoreContext::CheckModeEdit()
+{
+	if (const UPropertyAnimatorCoreHandlerBase* Handler = GetHandler())
+	{
+		bAllowModeEdit = Handler->IsAdditiveSupported();
+	}
+}
+
+void* UPropertyAnimatorCoreContext::GetConverterRulePtr(const UScriptStruct* InStruct)
+{
+	if (ConverterRule.IsValid() && ConverterRule.GetScriptStruct()->IsChildOf(InStruct))
+	{
+		return ConverterRule.GetMutableMemory();
+	}
+
+	return nullptr;
+}
+
+void UPropertyAnimatorCoreContext::Restore()
+{
+	if (OriginalPropertyValues.GetNumPropertiesInBag() == 0
+		&& DeltaPropertyValues.GetNumPropertiesInBag() == 0)
+	{
+		return;
+	}
+
+	UPropertyAnimatorCoreHandlerBase* Handler = GetHandler();
+
+	if (!Handler)
+	{
+		return;
+	}
+
+	if (Mode == EPropertyAnimatorCoreMode::Absolute)
+	{
+		for (const FPropertyAnimatorCoreData& ResolvedProperty : ResolveProperty(false))
+		{
+			// Reset original value
+			if (Handler->SetValue(ResolvedProperty, OriginalPropertyValues))
+			{
+				OriginalPropertyValues.RemovePropertyByName(FName(ResolvedProperty.GetPathHash()));
+			}
+		}
+
+		OriginalPropertyValues.Reset();
+	}
+	else
+	{
+		for (const FPropertyAnimatorCoreData& ResolvedProperty : ResolveProperty(false))
+		{
+			// Subtract delta value
+			if (Handler->SubtractValue(ResolvedProperty, DeltaPropertyValues))
+			{
+				DeltaPropertyValues.RemovePropertyByName(FName(ResolvedProperty.GetPathHash()));
+			}
+		}
+	}
+
+	DeltaPropertyValues.Reset();
+}
+
+void UPropertyAnimatorCoreContext::Save()
+{
+	UPropertyAnimatorCoreHandlerBase* Handler = GetHandler();
+
+	if (!Handler)
+	{
+		return;
+	}
+
+	for (const FPropertyAnimatorCoreData& PropertyData : ResolveProperty(false))
+	{
+		const FName Name(PropertyData.GetPathHash());
+		if (!OriginalPropertyValues.FindPropertyDescByName(Name))
+		{
+			const FProperty* Property = PropertyData.GetLeafProperty();
+			OriginalPropertyValues.AddProperty(Name, Property);
+
+			// Save original value
+			Handler->GetValue(PropertyData, OriginalPropertyValues);
+		}
+
+		if (!DeltaPropertyValues.FindPropertyDescByName(Name))
+		{
+			const FProperty* Property = PropertyData.GetLeafProperty();
+			DeltaPropertyValues.AddProperty(Name, Property);
+		}
+	}
+}
+
+TArray<FName> UPropertyAnimatorCoreContext::GetSupportedGroupNames() const
+{
+	TArray<FName> GroupNames
+	{
+		NAME_None
+	};
+
+	if (UPropertyAnimatorCoreBase* Animator = GetAnimator())
+	{
+		for (const UPropertyAnimatorCoreGroupBase* PropertyGroup : Animator->PropertyGroups)
+		{
+			if (PropertyGroup && PropertyGroup->IsPropertySupported(this))
+			{
+				GroupNames.Add(PropertyGroup->GetFName());
+			}
+		}
+	}
+
+	return GroupNames;
+}
+
+void UPropertyAnimatorCoreContext::OnAnimatedChanged()
+{
+	if (!bAnimated)
+	{
+		Restore();
+	}
+}
+
+void UPropertyAnimatorCoreContext::SetEvaluationResult(const FPropertyAnimatorCoreData& InResolvedProperty, const FInstancedPropertyBag& InEvaluatedValues)
+{
+	if (!IsAnimated())
+	{
+		return;
+	}
+
+	UPropertyAnimatorCoreHandlerBase* Handler = GetHandler();
+
+	if (!Handler)
+	{
+		return;
+	}
+
+	const FName PropertyName(InResolvedProperty.GetPathHash());
+
+	const FPropertyBagPropertyDesc* FromProperty = InEvaluatedValues.FindPropertyDescByName(PropertyName);
+	const FPropertyBagPropertyDesc* ToProperty = DeltaPropertyValues.FindPropertyDescByName(PropertyName);
+
+	if (UPropertyAnimatorCoreConverterBase* Converter = ConverterClass.GetDefaultObject())
+	{
+		if (!Converter->Convert(*FromProperty, InEvaluatedValues, *ToProperty, DeltaPropertyValues, ConverterRule.IsValid() ? &ConverterRule : nullptr))
+		{
+			return;
+		}
+	}
+	else
+	{
+		// Ids need to match for copy to be successful
+		const_cast<FPropertyBagPropertyDesc*>(FromProperty)->ID = ToProperty->ID;
+
+		DeltaPropertyValues.CopyMatchingValuesByID(InEvaluatedValues);
+	}
+
+	if (Mode == EPropertyAnimatorCoreMode::Absolute)
+	{
+		Handler->SetValue(InResolvedProperty, DeltaPropertyValues);
+		DeltaPropertyValues.RemovePropertyByName(PropertyName);
+	}
+	else
+	{
+		Handler->AddValue(InResolvedProperty, DeltaPropertyValues);
+	}
+}
diff --git a/Engine/Plugins/Experimental/PropertyAnimatorCore/Source/PropertyAnimatorCore/Private/TimeSources/PropertyAnimatorCoreManualTimeSource.cpp b/Engine/Plugins/Experimental/PropertyAnimatorCore/Source/PropertyAnimatorCore/Private/TimeSources/PropertyAnimatorCoreManualTimeSource.cpp
new file mode 100644
--- /dev/null
+++ b/Engine/Plugins/Experimental/PropertyAnimatorCore/Source/PropertyAnimatorCore/Private/TimeSources/PropertyAnimatorCoreManualTimeSource.cpp
@@ -0,0 +1,23 @@
+// Copyright Epic Games, Inc. All Rights Reserved.
+
+#include "TimeSources/PropertyAnimatorCoreManualTimeSource.h"
+
+void UPropertyAnimatorCoreManualTimeSource::SetCustomTime(float InTime)
+{
+	if (FMath::IsNearlyEqual(InTime, CustomTime))
+	{
+		return;
+	}
+
+	CustomTime = InTime;
+}
+
+double UPropertyAnimatorCoreManualTimeSource::GetTimeElapsed()
+{
+	return CustomTime;
+}
+
+bool UPropertyAnimatorCoreManualTimeSource::IsTimeSourceReady() const
+{
+	return true;
+}
\ No newline at end of file
diff --git a/Engine/Plugins/Experimental/PropertyAnimatorCore/Source/PropertyAnimatorCore/Private/TimeSources/PropertyAnimatorCoreRangeTimeSource.cpp b/Engine/Plugins/Experimental/PropertyAnimatorCore/Source/PropertyAnimatorCore/Private/TimeSources/PropertyAnimatorCoreRangeTimeSource.cpp
new file mode 100644
--- /dev/null
+++ b/Engine/Plugins/Experimental/PropertyAnimatorCore/Source/PropertyAnimatorCore/Private/TimeSources/PropertyAnimatorCoreRangeTimeSource.cpp
@@ -0,0 +1,58 @@
+// Copyright Epic Games, Inc. All Rights Reserved.
+
+#include "TimeSources/PropertyAnimatorCoreRangeTimeSource.h"
+
+void UPropertyAnimatorCoreRangeTimeSource::SetUseStartTime(bool bInUse)
+{
+	if (bUseStartTime == bInUse)
+	{
+		return;
+	}
+
+	bUseStartTime = bInUse;
+}
+
+void UPropertyAnimatorCoreRangeTimeSource::SetStartTime(double InStartTime)
+{
+	if (FMath::IsNearlyEqual(StartTime, InStartTime))
+	{
+		return;
+	}
+
+	StartTime = InStartTime;
+}
+
+void UPropertyAnimatorCoreRangeTimeSource::SetUseStopTime(bool bInUse)
+{
+	if (bUseStopTime == bInUse)
+	{
+		return;
+	}
+
+	bUseStopTime = bInUse;
+}
+
+void UPropertyAnimatorCoreRangeTimeSource::SetStopTime(double InStopTime)
+{
+	if (FMath::IsNearlyEqual(StopTime, InStopTime))
+	{
+		return;
+	}
+
+	StopTime = InStopTime;
+}
+
+bool UPropertyAnimatorCoreRangeTimeSource::IsValidTimeElapsed(double InTimeElapsed) const
+{
+	if (bUseStartTime && InTimeElapsed < StartTime)
+	{
+		return false;
+	}
+
+	if (bUseStopTime && InTimeElapsed > StopTime)
+	{
+		return false;
+	}
+
+	return true;
+}
diff --git a/Engine/Plugins/Experimental/PropertyAnimatorCore/Source/PropertyAnimatorCore/Private/TimeSources/PropertyAnimatorCoreTimeSourceBase.cpp b/Engine/Plugins/Experimental/PropertyAnimatorCore/Source/PropertyAnimatorCore/Private/TimeSources/PropertyAnimatorCoreTimeSourceBase.cpp
new file mode 100644
--- /dev/null
+++ b/Engine/Plugins/Experimental/PropertyAnimatorCore/Source/PropertyAnimatorCore/Private/TimeSources/PropertyAnimatorCoreTimeSourceBase.cpp
@@ -0,0 +1,49 @@
+// Copyright Epic Games, Inc. All Rights Reserved.
+
+#include "TimeSources/PropertyAnimatorCoreTimeSourceBase.h"
+
+#include "Animators/PropertyAnimatorCoreBase.h"
+
+void UPropertyAnimatorCoreTimeSourceBase::ActivateTimeSource()
+{
+	if (IsTimeSourceActive())
+	{
+		return;
+	}
+
+	bTimeSourceActive = true;
+	OnTimeSourceActive();
+}
+
+void UPropertyAnimatorCoreTimeSourceBase::DeactivateTimeSource()
+{
+	if (!IsTimeSourceActive())
+	{
+		return;
+	}
+
+	bTimeSourceActive = false;
+	OnTimeSourceInactive();
+}
+
+TOptional<double> UPropertyAnimatorCoreTimeSourceBase::GetConditionalTimeElapsed()
+{
+	if (!IsTimeSourceReady())
+	{
+		return TOptional<double>();
+	}
+
+	TimeElapsed = GetTimeElapsed();
+
+	if (!IsValidTimeElapsed(TimeElapsed))
+	{
+		return TOptional<double>();
+	}
+
+	return TimeElapsed;
+}
+
+UPropertyAnimatorCoreBase* UPropertyAnimatorCoreTimeSourceBase::GetAnimator() const
+{
+	return GetTypedOuter<UPropertyAnimatorCoreBase>();
+}
diff --git a/Engine/Plugins/Experimental/PropertyAnimatorCore/Source/PropertyAnimatorCore/Private/TimeSources/PropertyAnimatorCoreWorldTimeSource.cpp b/Engine/Plugins/Experimental/PropertyAnimatorCore/Source/PropertyAnimatorCore/Private/TimeSources/PropertyAnimatorCoreWorldTimeSource.cpp
new file mode 100644
--- /dev/null
+++ b/Engine/Plugins/Experimental/PropertyAnimatorCore/Source/PropertyAnimatorCore/Private/TimeSources/PropertyAnimatorCoreWorldTimeSource.cpp
@@ -0,0 +1,37 @@
+// Copyright Epic Games, Inc. All Rights Reserved.
+
+#include "TimeSources/PropertyAnimatorCoreWorldTimeSource.h"
+
+#include "Animators/PropertyAnimatorCoreBase.h"
+#include "Engine/World.h"
+
+double UPropertyAnimatorCoreWorldTimeSource::GetTimeElapsed()
+{
+	if (const UWorld* World = GetWorld())
+	{
+		return World->GetTimeSeconds();
+	}
+
+	return 0.f;
+}
+
+bool UPropertyAnimatorCoreWorldTimeSource::IsTimeSourceReady() const
+{
+	const UPropertyAnimatorCoreBase* Animator = GetAnimator();
+
+	if (!Animator)
+	{
+		return false;
+	}
+
+	const UWorld* World = Animator->GetWorld();
+
+	if (!World)
+	{
+		return false;
+	}
+
+	return World->IsEditorWorld()
+		|| World->IsPreviewWorld()
+		|| (World->IsGameWorld() && World->HasBegunPlay());
+}
diff --git a/Engine/Plugins/Experimental/PropertyAnimatorCore/Source/PropertyAnimatorCore/Public/Animators/PropertyAnimatorCoreBase.h b/Engine/Plugins/Experimental/PropertyAnimatorCore/Source/PropertyAnimatorCore/Public/Animators/PropertyAnimatorCoreBase.h
new file mode 100644
--- /dev/null
+++ b/Engine/Plugins/Experimental/PropertyAnimatorCore/Source/PropertyAnimatorCore/Public/Animators/PropertyAnimatorCoreBase.h
@@ -0,0 +1,310 @@
+// Copyright Epic Games, Inc. All Rights Reserved.
+
+#pragma once
+
+#include "Properties/PropertyAnimatorCoreContext.h"
+#include "Properties/PropertyAnimatorCoreData.h"
+#include "UObject/Object.h"
+#include "PropertyAnimatorCoreBase.generated.h"
+
+class UPropertyAnimatorCoreComponent;
+class UPropertyAnimatorCoreGroupBase;
+class UPropertyAnimatorCoreTimeSourceBase;
+
+/** Parameters to use for the current evaluation results */
+struct FPropertyAnimatorCoreEvaluationParameters
+{
+	/** Time elapsed generated by the active time source */
+	double TimeElapsed = 0.f;
+
+	/** Global magnitude for all animators */
+	float AnimatorsMagnitude = 1.f;
+};
+
+/** Abstract base class for any Animator, holds a set of linked properties */
+UCLASS(Abstract, EditInlineNew, AutoExpandCategories=("Animator"))
+class PROPERTYANIMATORCORE_API UPropertyAnimatorCoreBase : public UObject
+{
+	GENERATED_BODY()
+
+	friend class UPropertyAnimatorCoreComponent;
+	friend class UPropertyAnimatorCoreContext;
+	friend class UPropertyAnimatorCoreEditorStackCustomization;
+	friend class FPropertyAnimatorCoreEditorDetailCustomization;
+
+public:
+	DECLARE_MULTICAST_DELEGATE_OneParam(FOnAnimatorUpdated, UPropertyAnimatorCoreBase* /* InAnimator */)
+
+	/** Called when a Animator is created */
+	static FOnAnimatorUpdated OnAnimatorCreatedDelegate;
+
+	/** Called when a Animator is removed */
+	static FOnAnimatorUpdated OnAnimatorRemovedDelegate;
+
+	/** Called when a Animator is renamed */
+	static FOnAnimatorUpdated OnAnimatorRenamedDelegate;
+
+	DECLARE_MULTICAST_DELEGATE_TwoParams(FOnAnimatorPropertyUpdated, UPropertyAnimatorCoreBase* /* InAnimator */, const FPropertyAnimatorCoreData& /** InProperty */)
+
+	/** Called when a property is linked to a Animator */
+	static FOnAnimatorPropertyUpdated OnAnimatorPropertyLinkedDelegate;
+
+	/** Called when a property is unlinked to a Animator */
+	static FOnAnimatorPropertyUpdated OnAnimatorPropertyUnlinkedDelegate;
+
+	UPropertyAnimatorCoreBase();
+
+	AActor* GetAnimatorActor() const;
+	UPropertyAnimatorCoreComponent* GetAnimatorComponent() const;
+
+	/** Get the global animators magnitude */
+	float GetAnimatorComponentMagnitude() const;
+
+	/** Set the state of this animator */
+	void SetAnimatorEnabled(bool bInIsEnabled);
+	bool GetAnimatorEnabled() const
+	{
+		return bAnimatorEnabled;
+	}
+
+	/** Set the time source name to use */
+	void SetTimeSourceName(FName InTimeSourceName);
+	FName GetTimeSourceName() const
+	{
+		return TimeSourceName;
+	}
+
+	/** Get the active time source */
+	UPropertyAnimatorCoreTimeSourceBase* GetActiveTimeSource() const
+	{
+		return ActiveTimeSource;
+	}
+
+	/** Set the display name of this animator */
+	void SetAnimatorDisplayName(FName InName);
+	FString GetAnimatorDisplayName() const
+	{
+		return AnimatorDisplayName.ToString();
+	}
+
+	/** Gets the Animator original name */
+	FName GetAnimatorOriginalName() const;
+
+	/** Get all linked properties within this animator */
+	TSet<FPropertyAnimatorCoreData> GetLinkedProperties() const;
+
+	/** Get linked properties count within this animator */
+	int32 GetLinkedPropertiesCount() const;
+
+	/** Link property to this Animator to be able to drive it */
+	bool LinkProperty(const FPropertyAnimatorCoreData& InLinkProperty);
+
+	/** Unlink property from this Animator */
+	bool UnlinkProperty(const FPropertyAnimatorCoreData& InUnlinkProperty);
+
+	/** Checks if this Animator is controlling this property */
+	bool IsPropertyLinked(const FPropertyAnimatorCoreData& InPropertyData) const;
+
+	/** Checks if this animator is controlling all properties */
+	bool IsPropertiesLinked(const TSet<FPropertyAnimatorCoreData>& InProperties) const;
+
+	/** Returns all inner properties that are controlled by this Animator linked to member property */
+	TSet<FPropertyAnimatorCoreData> GetInnerPropertiesLinked(const FPropertyAnimatorCoreData& InPropertyData) const;
+
+	/**
+	 * Checks recursively for properties inside that member property that are supported by this Animator
+	 * if the member property is directly supported only returns that one, calls IsPropertySupported to check
+	 * When recursive is false, it will stop if member property or children is controllable and not look further in the chain
+	 */
+	bool GetPropertiesSupported(const FPropertyAnimatorCoreData& InPropertyData, TSet<FPropertyAnimatorCoreData>& OutProperties, bool bInRecursiveSearch = false) const;
+
+	/** Does not recurse, checks if this Animator supports this specific property, use GetPropertiesSupported to check if nested properties are supported */
+	bool IsPropertySupported(const FPropertyAnimatorCoreData& InPropertyData) const;
+
+	/** Override this to check if a property is directly supported by this animator */
+	virtual bool IsPropertyDirectlySupported(const FPropertyAnimatorCoreData& InPropertyData) const
+	{
+		return false;
+	}
+
+	/** Override this to check if a property is indirectly supported by this animator with conversions */
+	virtual bool IsPropertyIndirectlySupported(const FPropertyAnimatorCoreData& InPropertyData) const
+	{
+		return false;
+	}
+
+	/** Get the context for the linked property */
+	UPropertyAnimatorCoreContext* GetLinkedPropertyContext(const FPropertyAnimatorCoreData& InProperty) const;
+
+	/** Get the casted context for the linked property */
+	template<typename InContextClass
+		UE_REQUIRES(TIsDerivedFrom<InContextClass, UPropertyAnimatorCoreContext>::Value)>
+	InContextClass* GetLinkedPropertyContext(const FPropertyAnimatorCoreData& InProperty) const
+	{
+		return Cast<InContextClass>(GetLinkedPropertyContext(InProperty));
+	}
+
+protected:
+	//~ Begin UObject
+	virtual void PostLoad() override;
+#if WITH_EDITOR
+	virtual void PreEditUndo() override;
+	virtual void PostEditUndo() override;
+	virtual void PostEditChangeProperty(FPropertyChangedEvent& PropertyChangedEvent) override;
+#endif
+	//~ End UObject
+
+	/** Update display name based on linked properties */
+	void UpdateAnimatorDisplayName();
+
+	/** Use this to process each linked properties and resolve it, even virtual ones */
+	template<typename InContextClass
+		UE_REQUIRES(TIsDerivedFrom<InContextClass, UPropertyAnimatorCoreContext>::Value)>
+	bool ForEachLinkedProperty(TFunctionRef<bool(InContextClass*, const FPropertyAnimatorCoreData&)> InFunction, bool bInResolve = true) const
+	{
+		for (const TObjectPtr<UPropertyAnimatorCoreContext>& LinkedProperty : LinkedProperties)
+		{
+			if (InContextClass* PropertyContext = Cast<InContextClass>(LinkedProperty.Get()))
+			{
+				if (bInResolve)
+				{
+					for (const FPropertyAnimatorCoreData& ResolvedPropertyData : PropertyContext->ResolveProperty(false))
+					{
+						if (!ResolvedPropertyData.IsResolved())
+						{
+							continue;
+						}
+
+						if (!InFunction(PropertyContext, ResolvedPropertyData))
+						{
+							return false;
+						}
+					}
+				}
+				else
+				{
+					if (!InFunction(PropertyContext, LinkedProperty->GetAnimatedProperty()))
+					{
+						return false;
+					}
+				}
+			}
+		}
+
+		return true;
+	}
+
+	/** Used to evaluate linked properties, assign the result in the property bag and return true on success to update property value */
+	template<typename InContextClass
+		UE_REQUIRES(TIsDerivedFrom<InContextClass, UPropertyAnimatorCoreContext>::Value)>
+	void EvaluateEachLinkedProperty(
+		TFunctionRef<bool(
+			InContextClass* /** InPropertyContext */
+			, const FPropertyAnimatorCoreData& /** InResolvedProperty */
+			, FInstancedPropertyBag& /** OutEvaluation */)> InFunction
+		)
+	{
+		checkf(bEvaluatingProperties, TEXT("EvaluateEachLinkedProperty can only be called in EvaluateProperties"))
+
+		for (const TObjectPtr<UPropertyAnimatorCoreContext>& LinkedProperty : LinkedProperties)
+		{
+			if (InContextClass* PropertyContext = Cast<InContextClass>(LinkedProperty.Get()))
+			{
+				if (!PropertyContext->IsAnimated())
+				{
+					continue;
+				}
+
+				for (const FPropertyAnimatorCoreData& ResolvedPropertyData : PropertyContext->ResolveProperty(true))
+				{
+					if (!ResolvedPropertyData.IsResolved())
+					{
+						continue;
+					}
+
+					if (InFunction(PropertyContext, ResolvedPropertyData, EvaluatedPropertyValues))
+					{
+						PropertyContext->SetEvaluationResult(ResolvedPropertyData, EvaluatedPropertyValues);
+					}
+				}
+			}
+		}
+	}
+
+	virtual void OnAnimatorDisplayNameChanged() {}
+
+	virtual void OnAnimatorAdded() {}
+	virtual void OnAnimatorRemoved() {}
+
+	virtual void OnAnimatorEnabled();
+	virtual void OnAnimatorDisabled();
+
+	/** Returns the property context class to use */
+	virtual TSubclassOf<UPropertyAnimatorCoreContext> GetPropertyContextClass(const FPropertyAnimatorCoreData& InProperty);
+	virtual void OnPropertyLinked(UPropertyAnimatorCoreContext* InLinkedProperty) {}
+	virtual void OnPropertyUnlinked(UPropertyAnimatorCoreContext* InUnlinkedProperty) {}
+
+	/** Restore modified properties to original state */
+	virtual void RestoreProperties(bool bInForce = false);
+
+	/** Allocate and saves properties in the property bag */
+	virtual void SaveProperties();
+
+	/** Apply animators effect on linked properties */
+	virtual void EvaluateProperties(const FPropertyAnimatorCoreEvaluationParameters& InParameters) {}
+
+private:
+	/** Called by the component to evaluate this animator */
+	void EvaluateAnimator();
+
+	void OnPropertyGroupsChanged();
+
+	void OnAnimatorEnabledChanged();
+
+	void CleanLinkedProperties();
+
+	/** Clean invalid or renamed time sources */
+	void CleanTimeSources();
+
+	void OnTimeSourceNameChanged();
+
+	UPropertyAnimatorCoreTimeSourceBase* FindOrAddTimeSource(FName InTimeSourceName);
+
+	UFUNCTION()
+	TArray<FString> GetTimeSourceNames() const;
+
+	/** Enable control of properties linked to this Animator */
+	UPROPERTY(EditInstanceOnly, Getter="GetAnimatorEnabled", Setter="SetAnimatorEnabled", Category="Animator", meta=(DisplayPriority="0", AllowPrivateAccess="true"))
+	bool bAnimatorEnabled = true;
+
+	/** Display name as title property for component array, hide it but must be visible to editor for array title property */
+	UPROPERTY(VisibleInstanceOnly, Category="Animator", meta=(EditCondition="false", EditConditionHides))
+	FName AnimatorDisplayName;
+
+	/** Context for properties linked to this Animator */
+	UPROPERTY(EditInstanceOnly, Category="Animator", EditFixedSize, meta=(EditFixedOrder))
+	TArray<TObjectPtr<UPropertyAnimatorCoreContext>> LinkedProperties;
+
+	/** Groups for properties linked to this Animator */
+	UPROPERTY(EditInstanceOnly, Instanced, Category="Animator")
+	TArray<TObjectPtr<UPropertyAnimatorCoreGroupBase>> PropertyGroups;
+
+	/** The time source to use */
+	UPROPERTY(EditInstanceOnly, Setter="SetTimeSourceName", Getter="GetTimeSourceName", Category="Animator", meta=(GetOptions="GetTimeSourceNames"))
+	FName TimeSourceName = NAME_None;
+
+	/** Active time source with its options, determined by its name */
+	UPROPERTY(VisibleInstanceOnly, Instanced, Transient, DuplicateTransient, Category="Animator")
+	TObjectPtr<UPropertyAnimatorCoreTimeSourceBase> ActiveTimeSource;
+
+	/** The cached time source used by this Animator */
+	UPROPERTY()
+	TMap<FName, TObjectPtr<UPropertyAnimatorCoreTimeSourceBase>> TimeSourcesInstances;
+
+	/** Evaluated property container, reset on every update round */
+	UPROPERTY(Transient, DuplicateTransient, NonTransactional)
+	FInstancedPropertyBag EvaluatedPropertyValues;
+
+	/** Are we evaluating properties currently */
+	bool bEvaluatingProperties = false;
+};
\ No newline at end of file
diff --git a/Engine/Plugins/Experimental/PropertyAnimatorCore/Source/PropertyAnimatorCore/Public/Components/PropertyAnimatorCoreComponent.h b/Engine/Plugins/Experimental/PropertyAnimatorCore/Source/PropertyAnimatorCore/Public/Components/PropertyAnimatorCoreComponent.h
new file mode 100644
--- /dev/null
+++ b/Engine/Plugins/Experimental/PropertyAnimatorCore/Source/PropertyAnimatorCore/Public/Components/PropertyAnimatorCoreComponent.h
@@ -0,0 +1,109 @@
+// Copyright Epic Games, Inc. All Rights Reserved.
+
+#pragma once
+
+#include "Components/ActorComponent.h"
+#include "Animators/PropertyAnimatorCoreBase.h"
+#include "PropertyAnimatorCoreComponent.generated.h"
+
+/** A container for controllers that holds properties in this actor */
+UCLASS(ClassGroup=(Custom), AutoExpandCategories=("Animator"), HideCategories=("Activation", "Cooking", "AssetUserData", "Collision"), meta=(BlueprintSpawnableComponent))
+class PROPERTYANIMATORCORE_API UPropertyAnimatorCoreComponent : public UActorComponent
+{
+	GENERATED_BODY()
+
+	friend class UPropertyAnimatorCoreSubsystem;
+	friend class UPropertyAnimatorCoreEditorStackCustomization;
+
+public:
+	/** Create an instance of this component class and adds it to an actor */
+	static UPropertyAnimatorCoreComponent* FindOrAdd(AActor* InActor);
+
+	UPropertyAnimatorCoreComponent();
+
+	void SetAnimators(const TSet<TObjectPtr<UPropertyAnimatorCoreBase>>& InAnimators);
+	const TSet<TObjectPtr<UPropertyAnimatorCoreBase>>& GetAnimators() const
+	{
+		return Animators;
+	}
+
+	int32 GetAnimatorsCount() const
+	{
+		return Animators.Num();
+	}
+
+	/** Set the state of all animators in this component */
+	void SetAnimatorsEnabled(bool bInEnabled);
+	bool GetAnimatorsEnabled() const
+	{
+		return bAnimatorsEnabled;
+	}
+
+	/** Set the magnitude for all animators in this component */
+	void SetAnimatorsMagnitude(float InMagnitude);
+	float GetAnimatorsMagnitude() const
+	{
+		return AnimatorsMagnitude;
+	}
+
+	/** Process a function for each controller, stops when false is returned otherwise continue until the end */
+	void ForEachAnimator(TFunctionRef<bool(UPropertyAnimatorCoreBase*)> InFunction) const;
+
+protected:
+	static FName GetAnimatorName(const UPropertyAnimatorCoreBase* InAnimator);
+
+	//~ Begin UActorComponent
+	virtual void DestroyComponent(bool bPromoteChildren) override;
+	//~ End UActorComponent
+
+	//~ Begin UObject
+#if WITH_EDITOR
+	virtual void PostEditUndo() override;
+	virtual void PreEditChange(FProperty* PropertyAboutToChange) override;
+	virtual void PostEditChangeProperty(FPropertyChangedEvent& PropertyChangedEvent) override;
+#endif
+	//~ End UObject
+
+	/** Adds a new controller and returns it casted */
+	template<typename InAnimatorClass = UPropertyAnimatorCoreBase, typename = typename TEnableIf<TIsDerivedFrom<InAnimatorClass, UPropertyAnimatorCoreBase>::Value>::Type>
+	InAnimatorClass* AddAnimator()
+	{
+		const UClass* AnimatorClass = InAnimatorClass::StaticClass();
+		return Cast<InAnimatorClass>(AddAnimator(AnimatorClass));
+	}
+
+	/** Adds a new animator of that class */
+	UPropertyAnimatorCoreBase* AddAnimator(const UClass* InAnimatorClass);
+
+	/** Removes an existing animator */
+	bool RemoveAnimator(UPropertyAnimatorCoreBase* InAnimator);
+
+	/** Change global state for animators */
+	void OnAnimatorsSetEnabled(const UWorld* InWorld, bool bInEnabled, bool bInTransact);
+
+	void OnAnimatorsChanged();
+
+	void OnAnimatorsEnabledChanged();
+
+	/** Checks if this component animators should tick */
+	bool ShouldAnimatorsTick() const;
+
+	/** Animators linked to this actor, they contain only properties within this actor */
+	UPROPERTY(EditAnywhere, BlueprintReadWrite, Instanced, Getter="GetAnimators", Setter="SetAnimators", Category="Animator", meta=(TitleProperty="AnimatorDisplayName"))
+	TSet<TObjectPtr<UPropertyAnimatorCoreBase>> Animators;
+
+	/** Global state for all animators controlled by this component */
+	UPROPERTY(EditInstanceOnly, Getter="GetAnimatorsEnabled", Setter="SetAnimatorsEnabled", Category="Animator", meta=(DisplayPriority="0", AllowPrivateAccess="true"))
+	bool bAnimatorsEnabled = true;
+
+	/** Global magnitude for all animators controlled by this component */
+	UPROPERTY(EditInstanceOnly, Getter, Setter, Category="Animator", meta=(ClampMin="0", ClampMax="1", UIMin="0", UIMax="1", AllowPrivateAccess="true"))
+	float AnimatorsMagnitude = 1.f;
+
+private:
+	virtual void TickComponent(float InDeltaTime, ELevelTick InTickType, FActorComponentTickFunction* InThisTickFunction) override;
+
+	/** Transient copy of animators set when changes are detected to see the diff only */
+	UPROPERTY(Transient, DuplicateTransient, TextExportTransient)
+	TSet<TObjectPtr<UPropertyAnimatorCoreBase>> AnimatorsInternal;
+};
\ No newline at end of file
diff --git a/Engine/Plugins/Experimental/PropertyAnimatorCore/Source/PropertyAnimatorCore/Public/Presets/PropertyAnimatorCorePresetBase.h b/Engine/Plugins/Experimental/PropertyAnimatorCore/Source/PropertyAnimatorCore/Public/Presets/PropertyAnimatorCorePresetBase.h
new file mode 100644
--- /dev/null
+++ b/Engine/Plugins/Experimental/PropertyAnimatorCore/Source/PropertyAnimatorCore/Public/Presets/PropertyAnimatorCorePresetBase.h
@@ -0,0 +1,62 @@
+// Copyright Epic Games, Inc. All Rights Reserved.
+
+#pragma once
+
+#include "Properties/PropertyAnimatorCoreData.h"
+#include "PropertyAnimatorCorePresetBase.generated.h"
+
+class AActor;
+class UPropertyAnimatorCoreBase;
+
+/**
+ * Abstract Base class to define preset for animators with custom properties and options
+ * Will get registered automatically by the subsystem
+ * Should remain transient and stateless
+ */
+UCLASS(MinimalAPI, Abstract, Transient)
+class UPropertyAnimatorCorePresetBase : public UObject
+{
+	GENERATED_BODY()
+
+public:
+	PROPERTYANIMATORCORE_API UPropertyAnimatorCorePresetBase()
+		: UPropertyAnimatorCorePresetBase(NAME_None)
+	{}
+
+	PROPERTYANIMATORCORE_API UPropertyAnimatorCorePresetBase(FName InPresetName)
+		: PresetName(InPresetName)
+	{}
+
+	PROPERTYANIMATORCORE_API FName GetPresetName() const
+	{
+		return PresetName;
+	}
+
+	PROPERTYANIMATORCORE_API FString GetPresetDisplayName() const;
+
+	/** Get the preset properties for that actor */
+	PROPERTYANIMATORCORE_API virtual void GetPresetProperties(const AActor* InActor, const UPropertyAnimatorCoreBase* InAnimator, TSet<FPropertyAnimatorCoreData>& OutProperties) const {}
+
+	/** Get the preset properties for that actor but only supported ones for that animator */
+	PROPERTYANIMATORCORE_API virtual void GetSupportedPresetProperties(const AActor* InActor, const UPropertyAnimatorCoreBase* InAnimator, TSet<FPropertyAnimatorCoreData>& OutProperties) const;
+
+	/** Checks if the preset is supported on that actor and animator */
+	PROPERTYANIMATORCORE_API virtual bool IsPresetSupported(const AActor* InActor, const UPropertyAnimatorCoreBase* InAnimator) const;
+
+	/** Checks if the preset is applied to this animator */
+	PROPERTYANIMATORCORE_API virtual bool IsPresetApplied(const UPropertyAnimatorCoreBase* InAnimator) const;
+
+	/** Applies this preset on the newly created animator */
+	PROPERTYANIMATORCORE_API virtual bool ApplyPreset(UPropertyAnimatorCoreBase* InAnimator);
+
+	/** Un applies this preset on an animator */
+	PROPERTYANIMATORCORE_API virtual bool UnapplyPreset(UPropertyAnimatorCoreBase* InAnimator);
+
+	/** Called when this preset is applied on the animator */
+	PROPERTYANIMATORCORE_API virtual void OnPresetApplied(UPropertyAnimatorCoreBase* InAnimator, const TSet<FPropertyAnimatorCoreData>& InProperties) {}
+
+private:
+	/** Name used to display this preset to the user */
+	UPROPERTY(Transient)
+	FName PresetName;
+};
\ No newline at end of file
diff --git a/Engine/Plugins/Experimental/PropertyAnimatorCore/Source/PropertyAnimatorCore/Public/Properties/PropertyAnimatorCoreContext.h b/Engine/Plugins/Experimental/PropertyAnimatorCore/Source/PropertyAnimatorCore/Public/Properties/PropertyAnimatorCoreContext.h
new file mode 100644
--- /dev/null
+++ b/Engine/Plugins/Experimental/PropertyAnimatorCore/Source/PropertyAnimatorCore/Public/Properties/PropertyAnimatorCoreContext.h
@@ -0,0 +1,168 @@
+// Copyright Epic Games, Inc. All Rights Reserved.
+
+#pragma once
+
+#include "PropertyAnimatorCoreData.h"
+#include "PropertyBag.h"
+#include "UObject/Object.h"
+#include "PropertyAnimatorCoreContext.generated.h"
+
+struct FInstancedStruct;
+class UPropertyAnimatorCoreBase;
+class UPropertyAnimatorCoreConverterBase;
+class UPropertyAnimatorCoreGroupBase;
+class UScriptStruct;
+
+/** Mode supported for properties value */
+UENUM(BlueprintType)
+enum class EPropertyAnimatorCoreMode : uint8
+{
+	Absolute,
+	Additive,
+};
+
+/** Context for properties linked to an animator */
+UCLASS(MinimalAPI, BlueprintType)
+class UPropertyAnimatorCoreContext : public UObject
+{
+	GENERATED_BODY()
+
+	friend class UPropertyAnimatorCoreBase;
+	friend class FPropertyAnimatorCoreEditorContextTypeCustomization;
+
+public:
+	const FPropertyAnimatorCoreData& GetAnimatedProperty() const
+	{
+		return AnimatedProperty;
+	}
+
+	PROPERTYANIMATORCORE_API UPropertyAnimatorCoreBase* GetAnimator() const;
+
+	/** Get the handler responsible for this property type */
+	PROPERTYANIMATORCORE_API UPropertyAnimatorCoreHandlerBase* GetHandler() const;
+
+	/** Get the active group of this property */
+	UPropertyAnimatorCoreGroupBase* GetGroup() const
+	{
+		return Group;
+	}
+
+	PROPERTYANIMATORCORE_API void SetAnimated(bool bInAnimated);
+	bool IsAnimated() const
+	{
+		return bAnimated;
+	}
+
+	PROPERTYANIMATORCORE_API void SetMode(EPropertyAnimatorCoreMode InMode);
+	EPropertyAnimatorCoreMode GetMode() const
+	{
+		return Mode;
+	}
+
+	PROPERTYANIMATORCORE_API void SetConverterClass(TSubclassOf<UPropertyAnimatorCoreConverterBase> InConverterClass);
+	TSubclassOf<UPropertyAnimatorCoreConverterBase> GetConverterClass() const
+	{
+		return ConverterClass;
+	}
+
+	PROPERTYANIMATORCORE_API void SetGroupName(FName InGroupName);
+	FName GetGroupName() const
+	{
+		return GroupName;
+	}
+
+	/** Get converter rule if any */
+	template <typename InRuleType
+		UE_REQUIRES(TModels_V<CStaticStructProvider, InRuleType>)>
+	InRuleType* GetConverterRule()
+	{
+		return static_cast<InRuleType*>(GetConverterRulePtr(InRuleType::StaticStruct()));
+	}
+
+protected:
+	//~ Begin UObject
+	PROPERTYANIMATORCORE_API virtual void PostLoad() override;
+#if WITH_EDITOR
+	PROPERTYANIMATORCORE_API virtual void PreEditChange(FProperty* InPropertyAboutToChange) override;
+	PROPERTYANIMATORCORE_API virtual void PostEditChangeProperty(FPropertyChangedEvent& InPropertyChangedEvent) override;
+#endif
+	//~ End UObject
+
+	/** Called once, when the property is linked to this context */
+	virtual void OnAnimatedPropertyLinked() {}
+
+private:
+	void ConstructInternal(const FPropertyAnimatorCoreData& InProperty);
+
+	PROPERTYANIMATORCORE_API void* GetConverterRulePtr(const UScriptStruct* InStruct);
+
+	void CheckModeEdit();
+	void OnAnimatedChanged();
+	void OnModeChanged();
+	void OnGroupNameChanged();
+
+	/** Sets the evaluation result for the resolved property */
+	PROPERTYANIMATORCORE_API void SetEvaluationResult(const FPropertyAnimatorCoreData& InResolvedProperty, const FInstancedPropertyBag& InEvaluatedValues);
+
+	/** Use this to resolve virtual linked property */
+	PROPERTYANIMATORCORE_API TArray<FPropertyAnimatorCoreData> ResolveProperty(bool bInForEvaluation) const;
+
+	/** Restore property based on mode */
+	void Restore();
+
+	/** Allocate and save properties */
+	void Save();
+
+	void SetGroup(UPropertyAnimatorCoreGroupBase* InGroup);
+
+	bool IsResolvable() const;
+	bool IsConverted() const;
+
+	/** Get the supported group names that can manage this property */
+	UFUNCTION()
+	TArray<FName> GetSupportedGroupNames() const;
+
+	/** Animation is enabled for this property */
+	UPROPERTY(EditInstanceOnly, Setter="SetAnimated", Getter="IsAnimated", Category="Animator", meta=(AllowPrivateAccess="true"))
+	bool bAnimated = true;
+
+	/** Edit condition for modes (Force hide otherwise appears inline for mode) */
+	UPROPERTY(Transient, VisibleInstanceOnly, Category="Animator", meta=(EditCondition="false", EditConditionHides))
+	bool bAllowModeEdit = true;
+
+	/** Current mode used for this property */
+	UPROPERTY(EditInstanceOnly, Setter, Getter, Category="Animator", meta=(EditCondition="bAllowModeEdit", EditConditionHides, AllowPrivateAccess="true"))
+	EPropertyAnimatorCoreMode Mode = EPropertyAnimatorCoreMode::Absolute;
+
+	/** If a converter is used, rules may be used to convert the property */
+	UPROPERTY(EditInstanceOnly, Category="Animator", meta=(AllowPrivateAccess="true"))
+	FInstancedStruct ConverterRule;
+
+	/** The unique group name that manages this property */
+	UPROPERTY(EditInstanceOnly, Category="Animator", Setter, Getter, meta=(GetOptions="GetSupportedGroupNames", AllowPrivateAccess="true"))
+	FName GroupName = NAME_None;
+
+	/** Active group of this property */
+	UPROPERTY()
+	TObjectPtr<UPropertyAnimatorCoreGroupBase> Group;
+
+	/** Store original property values for resolved properties */
+	UPROPERTY(DuplicateTransient, NonTransactional)
+	FInstancedPropertyBag OriginalPropertyValues;
+
+	/** Store delta property values for resolved properties */
+	UPROPERTY(DuplicateTransient, NonTransactional)
+	FInstancedPropertyBag DeltaPropertyValues;
+
+	/** Converter class used for this property */
+	UPROPERTY()
+	TSubclassOf<UPropertyAnimatorCoreConverterBase> ConverterClass;
+
+	/** Used to access property value and update it */
+	UPROPERTY(Transient)
+	TWeakObjectPtr<UPropertyAnimatorCoreHandlerBase> HandlerWeak;
+
+	/** Animated property linked to this options */
+	UPROPERTY()
+	FPropertyAnimatorCoreData AnimatedProperty;
+};
\ No newline at end of file
diff --git a/Engine/Plugins/Experimental/PropertyAnimatorCore/Source/PropertyAnimatorCore/Public/Properties/PropertyAnimatorCoreGroupBase.h b/Engine/Plugins/Experimental/PropertyAnimatorCore/Source/PropertyAnimatorCore/Public/Properties/PropertyAnimatorCoreGroupBase.h
new file mode 100644
--- /dev/null
+++ b/Engine/Plugins/Experimental/PropertyAnimatorCore/Source/PropertyAnimatorCore/Public/Properties/PropertyAnimatorCoreGroupBase.h
@@ -0,0 +1,40 @@
+// Copyright Epic Games, Inc. All Rights Reserved.
+
+#pragma once
+
+#include "Properties/PropertyAnimatorCoreData.h"
+#include "UObject/NameTypes.h"
+#include "UObject/Object.h"
+#include "PropertyAnimatorCoreGroupBase.generated.h"
+
+class UPropertyAnimatorCoreBase;
+class UPropertyAnimatorCoreContext;
+
+/** Group that controls animated properties in an animator */
+UCLASS(MinimalAPI, BlueprintType, Abstract, EditInlineNew)
+class UPropertyAnimatorCoreGroupBase : public UObject
+{
+	GENERATED_BODY()
+
+public:
+	UPropertyAnimatorCoreGroupBase()
+		: UPropertyAnimatorCoreGroupBase(NAME_None)
+	{}
+
+	UPropertyAnimatorCoreGroupBase(FName InName)
+		: GroupName(InName)
+	{}
+
+	/** Manages the properties owned by this group */
+	virtual void ManageProperties(const UPropertyAnimatorCoreContext* InContext, TArray<FPropertyAnimatorCoreData>& InOutProperties) {}
+
+	/** Checks if this group supports the context */
+	virtual bool IsPropertySupported(const UPropertyAnimatorCoreContext* InContext) const
+	{
+		return false;
+	}
+
+private:
+	UPROPERTY()
+	FName GroupName = NAME_None;
+};
\ No newline at end of file
diff --git a/Engine/Plugins/Experimental/PropertyAnimatorCore/Source/PropertyAnimatorCore/Public/Properties/PropertyAnimatorCoreResolver.h b/Engine/Plugins/Experimental/PropertyAnimatorCore/Source/PropertyAnimatorCore/Public/Properties/PropertyAnimatorCoreResolver.h
new file mode 100644
--- /dev/null
+++ b/Engine/Plugins/Experimental/PropertyAnimatorCore/Source/PropertyAnimatorCore/Public/Properties/PropertyAnimatorCoreResolver.h
@@ -0,0 +1,54 @@
+// Copyright Epic Games, Inc. All Rights Reserved.
+
+#pragma once
+
+#include "Properties/PropertyAnimatorCoreData.h"
+#include "UObject/NameTypes.h"
+#include "UObject/Object.h"
+#include "PropertyAnimatorCoreResolver.generated.h"
+
+/**
+ * Base class to find properties hidden or not reachable,
+ * allows to discover resolvable properties for specific actors/components/objects
+ * that we cannot reach or are transient, will be resolved when needed
+ * Resolvers should remain transient and stateless
+ */
+UCLASS(Abstract, Transient)
+class PROPERTYANIMATORCORE_API UPropertyAnimatorCoreResolver : public UObject
+{
+	GENERATED_BODY()
+
+public:
+	UPropertyAnimatorCoreResolver()
+		: UPropertyAnimatorCoreResolver(NAME_None, false)
+	{}
+
+	UPropertyAnimatorCoreResolver(FName InResolverName, bool bInIsRange)
+		: ResolverName(InResolverName)
+		, bIsRange(bInIsRange)
+	{}
+
+	/** Get properties found inside parent property */
+	virtual void GetResolvableProperties(const FPropertyAnimatorCoreData& InParentProperty, TSet<FPropertyAnimatorCoreData>& OutProperties) {}
+
+	/** Called when we actually need the underlying properties */
+	virtual void ResolveProperties(const FPropertyAnimatorCoreData& InTemplateProperty, TArray<FPropertyAnimatorCoreData>& OutProperties) {}
+
+	FName GetResolverName() const
+	{
+		return ResolverName;
+	}
+
+	/** Whether this property when resolved should be treated as a range */
+	bool IsRange() const
+	{
+		return bIsRange;
+	}
+
+private:
+	UPROPERTY()
+	FName ResolverName = NAME_None;
+
+	UPROPERTY()
+	bool bIsRange = false;
+};
\ No newline at end of file
diff --git a/Engine/Plugins/Experimental/PropertyAnimatorCore/Source/PropertyAnimatorCore/Public/Subsystems/PropertyAnimatorCoreSubsystem.h b/Engine/Plugins/Experimental/PropertyAnimatorCore/Source/PropertyAnimatorCore/Public/Subsystems/PropertyAnimatorCoreSubsystem.h
new file mode 100644
--- /dev/null
+++ b/Engine/Plugins/Experimental/PropertyAnimatorCore/Source/PropertyAnimatorCore/Public/Subsystems/PropertyAnimatorCoreSubsystem.h
@@ -0,0 +1,218 @@
+// Copyright Epic Games, Inc. All Rights Reserved.
+
+#pragma once
+
+#include "Animators/PropertyAnimatorCoreBase.h"
+#include "Presets/PropertyAnimatorCorePresetBase.h"
+#include "Properties/PropertyAnimatorCoreData.h"
+#include "Subsystems/EngineSubsystem.h"
+#include "UObject/WeakObjectPtrTemplates.h"
+#include "PropertyAnimatorCoreSubsystem.generated.h"
+
+class UFunction;
+class UPropertyAnimatorCoreConverterBase;
+class UPropertyAnimatorCoreBase;
+class UPropertyAnimatorCoreHandlerBase;
+class UPropertyAnimatorCorePresetBase;
+class UPropertyAnimatorCoreResolver;
+class UPropertyAnimatorCoreTimeSourceBase;
+
+/** This subsystem handle all property animators */
+UCLASS()
+class PROPERTYANIMATORCORE_API UPropertyAnimatorCoreSubsystem : public UEngineSubsystem
+{
+	GENERATED_BODY()
+
+	friend class UPropertyAnimatorCoreComponent;
+
+public:
+	/** Get this subsystem instance */
+	static UPropertyAnimatorCoreSubsystem* Get();
+
+	//~ Begin USubsystem
+	virtual void Initialize(FSubsystemCollectionBase& Collection) override;
+	virtual void Deinitialize() override;
+	//~ End USubsystem
+
+	/** Register the property controller class to allow its usage */
+	bool RegisterAnimatorClass(const UClass* InPropertyControllerClass);
+
+	/** Unregister the property controller class to disallow its usage */
+	bool UnregisterAnimatorClass(const UClass* InPropertyControllerClass);
+
+	/** Checks if the property controller class is already registered */
+	bool IsAnimatorClassRegistered(const UClass* InPropertyControllerClass) const;
+
+	/** Gets the animator CDO registered from the class */
+	UPropertyAnimatorCoreBase* GetAnimatorRegistered(const UClass* InAnimatorClass) const;
+
+	/** Returns true if any controller is able to control that property or nested otherwise false */
+	bool IsPropertySupported(const FPropertyAnimatorCoreData& InPropertyData, bool bInCheckNestedProperties = true) const;
+
+	/** Find all animators linked to the property */
+	TSet<UPropertyAnimatorCoreBase*> GetPropertyLinkedAnimators(const FPropertyAnimatorCoreData& InPropertyData) const;
+
+	/** Returns a set of existing property controller objects in owner that supports that property */
+	TSet<UPropertyAnimatorCoreBase*> GetExistingAnimators(const FPropertyAnimatorCoreData& InPropertyData) const;
+	TSet<UPropertyAnimatorCoreBase*> GetExistingAnimators(const AActor* InActor) const;
+
+	/** Returns a set of property controller CDO that supports that property */
+	TSet<UPropertyAnimatorCoreBase*> GetAvailableAnimators(const FPropertyAnimatorCoreData* InPropertyData) const;
+	TSet<UPropertyAnimatorCoreBase*> GetAvailableAnimators() const;
+
+	/** Register the property handler class to allow its usage */
+	bool RegisterHandlerClass(const UClass* InHandlerClass);
+
+	/** Unregister the property handler class to disallow its usage */
+	bool UnregisterHandlerClass(const UClass* InHandlerClass);
+
+	/** Checks if the property handler class is already registered */
+	bool IsHandlerClassRegistered(const UClass* InHandlerClass) const;
+
+	/** Gets a property handler for this property */
+	UPropertyAnimatorCoreHandlerBase* GetHandler(const FPropertyAnimatorCoreData& InPropertyData) const;
+
+	/** Register a resolver for custom properties */
+	bool RegisterResolverClass(const UClass* InResolverClass);
+
+	/** Unregister a resolver */
+	bool UnregisterResolverClass(const UClass* InResolverClass);
+
+	/** Is this resolver registered */
+	bool IsResolverClassRegistered(const UClass* InResolverClass) const;
+
+	/** Get all resolvable properties for an object property */
+	void GetResolvableProperties(const FPropertyAnimatorCoreData& InPropertyData, TSet<FPropertyAnimatorCoreData>& OutProperties) const;
+
+	/** Register a time source class to control clock for animators */
+	bool RegisterTimeSourceClass(UClass* InTimeSourceClass);
+
+	/** Unregister a time source class */
+	bool UnregisterTimeSourceClass(UClass* InTimeSourceClass);
+
+	/** Check time source class is registered */
+	bool IsTimeSourceClassRegistered(UClass* InTimeSourceClass) const;
+
+	/** Get all time sources available */
+	TArray<FName> GetTimeSourceNames() const;
+
+	/** Get a registered time source using its name */
+	UPropertyAnimatorCoreTimeSourceBase* GetTimeSource(FName InTimeSourceName) const;
+
+	/** Create a new time source for an animator */
+	UPropertyAnimatorCoreTimeSourceBase* CreateNewTimeSource(FName InTimeSourceName, UPropertyAnimatorCoreBase* InAnimator);
+
+	/** Register a preset class */
+	bool RegisterPresetClass(const UClass* InPresetClass);
+
+	/** Unregister a preset class */
+	bool UnregisterPresetClass(const UClass* InPresetClass);
+
+	/** Is this preset class registered */
+	bool IsPresetClassRegistered(const UClass* InPresetClass) const;
+
+	/** Gets all supported presets for a specific animator and actor */
+	TSet<UPropertyAnimatorCorePresetBase*> GetSupportedPresets(const AActor* InActor, const UPropertyAnimatorCoreBase* InAnimator) const;
+
+	/** Get all registered preset available */
+	TSet<UPropertyAnimatorCorePresetBase*> GetAvailablePresets() const;
+
+	bool RegisterSetterResolver(FName InPropertyName, TFunction<UFunction*(const UObject*)>&& InFunction);
+
+	bool UnregisterSetterResolver(FName InPropertyName);
+
+	UFunction* ResolveSetter(FName InPropertyName, const UObject* InOwner);
+
+	/** Register a converter class */
+	bool RegisterConverterClass(const UClass* InConverterClass);
+
+	/** Unregister a converter class */
+	bool UnregisterConverterClass(const UClass* InConverterClass);
+
+	/** Is this converter class registered */
+	bool IsConverterClassRegistered(const UClass* InConverterClass);
+
+	/** Checks if any converter supports the type conversion */
+	bool IsConversionSupported(const FPropertyBagPropertyDesc& InFromProperty, const FPropertyBagPropertyDesc& InToProperty);
+
+	/** Finds suitable converters for a type conversion */
+	TSet<UPropertyAnimatorCoreConverterBase*> GetSupportedConverters(const FPropertyBagPropertyDesc& InFromProperty, const FPropertyBagPropertyDesc& InToProperty) const;
+
+	/** Create an animator of specific class for an actor */
+	UPropertyAnimatorCoreBase* CreateAnimator(AActor* InActor, const UClass* InAnimatorClass, UPropertyAnimatorCorePresetBase* InPreset = nullptr, bool bInTransact = false) const;
+
+	/** Create animators of specific class for actors */
+	TSet<UPropertyAnimatorCoreBase*> CreateAnimators(const TSet<AActor*>& InActors, const UClass* InAnimatorClass, UPropertyAnimatorCorePresetBase* InPreset = nullptr, bool bInTransact = false) const;
+
+	/** Removes a animator bound to an owner */
+	bool RemoveAnimator(UPropertyAnimatorCoreBase* InAnimator, bool bInTransact = false) const;
+
+	/** Removes animators from their owner */
+	bool RemoveAnimators(const TSet<UPropertyAnimatorCoreBase*> InAnimators, bool bInTransact = false) const;
+
+	/** Apply a preset on an existing animator */
+	bool ApplyAnimatorPreset(UPropertyAnimatorCoreBase* InAnimator, UPropertyAnimatorCorePresetBase* InPreset, bool bInTransact = false);
+
+	/** Unapply a preset from an existing animator */
+	bool UnapplyAnimatorPreset(UPropertyAnimatorCoreBase* InAnimator, UPropertyAnimatorCorePresetBase* InPreset, bool bInTransact = false);
+
+	/** Link a property to an existing animator */
+	bool LinkAnimatorProperty(UPropertyAnimatorCoreBase* InAnimator, FPropertyAnimatorCoreData& InProperty, bool bInTransact = false);
+	bool LinkAnimatorProperties(UPropertyAnimatorCoreBase* InAnimator, const TSet<FPropertyAnimatorCoreData>& InProperties, bool bInTransact = false);
+
+	/** Unlink a property from an existing animator */
+	bool UnlinkAnimatorProperty(UPropertyAnimatorCoreBase* InAnimator, FPropertyAnimatorCoreData& InProperty, bool bInTransact = false);
+	bool UnlinkAnimatorProperties(UPropertyAnimatorCoreBase* InAnimator, const TSet<FPropertyAnimatorCoreData>& InProperties, bool bInTransact = false);
+
+	/** Set the enabled state of animators attached to actors, will disable state globally on the component */
+	void SetActorAnimatorsEnabled(const TSet<AActor*>& InActors, bool bInEnabled, bool bInTransact = false);
+
+	/** Set the enabled state of animators in a world, will disable state globally on the component */
+	void SetLevelAnimatorsEnabled(const UWorld* InWorld, bool bInEnabled, bool bInTransact = false);
+
+	/** Set the enabled state of animators provided */
+	void SetAnimatorsEnabled(const TSet<UPropertyAnimatorCoreBase*>& InAnimators, bool bInEnabled, bool bInTransact = false);
+
+protected:
+	/** Delegate to change state of animators in a world */
+	DECLARE_MULTICAST_DELEGATE_ThreeParams(FOnAnimatorsSetEnabled, const UWorld* /** InAnimatorWorld */, bool /** bInAnimatorEnabled */, bool /** bInTransact */)
+	static FOnAnimatorsSetEnabled OnAnimatorsSetEnabledDelegate;
+
+	/**
+	 * Scan for children of each of the following classes and registers their CDO:
+	 * 1. UPropertyAnimatorCoreBase
+	 * 2. UPropertyAnimatorCoreHandlerBase
+	 * 3. UPropertyAnimatorCoreResolver
+	 * 4. UPropertyAnimatorCoreTimeSourceBase
+	 * 5. UPropertyAnimatorCorePresetBase
+	 * 6. UPropertyAnimatorCoreConverterBase
+	 */
+	void RegisterAnimatorClasses();
+
+	/** Time sources available to use with animators */
+	UPROPERTY()
+	TSet<TWeakObjectPtr<UPropertyAnimatorCoreTimeSourceBase>> TimeSourcesWeak;
+
+	/** Animators available to link properties to */
+	UPROPERTY()
+	TSet<TWeakObjectPtr<UPropertyAnimatorCoreBase>> AnimatorsWeak;
+
+	/** Handlers are used to set/get same type properties and reuse logic */
+	UPROPERTY()
+	TSet<TWeakObjectPtr<UPropertyAnimatorCoreHandlerBase>> HandlersWeak;
+
+	/** Resolvers find properties to let user control them when they are unreachable/hidden */
+	UPROPERTY()
+	TSet<TWeakObjectPtr<UPropertyAnimatorCoreResolver>> ResolversWeak;
+
+	/** Presets available to apply on animator */
+	UPROPERTY()
+	TSet<TWeakObjectPtr<UPropertyAnimatorCorePresetBase>> PresetsWeak;
+
+	/** Converters available to transform a type to another type */
+	UPROPERTY()
+	TSet<TWeakObjectPtr<UPropertyAnimatorCoreConverterBase>> ConvertersWeak;
+
+	/** Some property and their setter cannot be identified automatically, use manual setter resolvers */
+	TMap<FName, TFunction<UFunction*(const UObject*)>> SetterResolvers;
+};
diff --git a/Engine/Plugins/Experimental/PropertyAnimatorCore/Source/PropertyAnimatorCore/Public/TimeSources/PropertyAnimatorCoreManualTimeSource.h b/Engine/Plugins/Experimental/PropertyAnimatorCore/Source/PropertyAnimatorCore/Public/TimeSources/PropertyAnimatorCoreManualTimeSource.h
new file mode 100644
--- /dev/null
+++ b/Engine/Plugins/Experimental/PropertyAnimatorCore/Source/PropertyAnimatorCore/Public/TimeSources/PropertyAnimatorCoreManualTimeSource.h
@@ -0,0 +1,34 @@
+// Copyright Epic Games, Inc. All Rights Reserved.
+
+#pragma once
+
+#include "PropertyAnimatorCoreTimeSourceBase.h"
+#include "PropertyAnimatorCoreManualTimeSource.generated.h"
+
+UCLASS()
+class PROPERTYANIMATORCORE_API UPropertyAnimatorCoreManualTimeSource : public UPropertyAnimatorCoreTimeSourceBase
+{
+	GENERATED_BODY()
+
+public:
+	UPropertyAnimatorCoreManualTimeSource()
+		: UPropertyAnimatorCoreTimeSourceBase(TEXT("Manual"))
+	{}
+
+	UFUNCTION()
+	void SetCustomTime(float InTime);
+	float GetCustomTime() const
+	{
+		return CustomTime;
+	}
+
+	//~ Begin UPropertyAnimatorTimeSourceBase
+	virtual double GetTimeElapsed() override;
+	virtual bool IsTimeSourceReady() const override;
+	//~ End UPropertyAnimatorTimeSourceBase
+
+protected:
+	/** Allows you to drive controllers with this float */
+	UPROPERTY(EditInstanceOnly, BlueprintReadWrite, Setter="SetCustomTime", Getter="GetCustomTime", Category="Animator")
+	float CustomTime = 0.f;
+};
\ No newline at end of file
diff --git a/Engine/Plugins/Experimental/PropertyAnimatorCore/Source/PropertyAnimatorCore/Public/TimeSources/PropertyAnimatorCoreRangeTimeSource.h b/Engine/Plugins/Experimental/PropertyAnimatorCore/Source/PropertyAnimatorCore/Public/TimeSources/PropertyAnimatorCoreRangeTimeSource.h
new file mode 100644
--- /dev/null
+++ b/Engine/Plugins/Experimental/PropertyAnimatorCore/Source/PropertyAnimatorCore/Public/TimeSources/PropertyAnimatorCoreRangeTimeSource.h
@@ -0,0 +1,69 @@
+// Copyright Epic Games, Inc. All Rights Reserved.
+
+#pragma once
+
+#include "PropertyAnimatorCoreTimeSourceBase.h"
+#include "PropertyAnimatorCoreRangeTimeSource.generated.h"
+
+/**
+ * Abstract range time source allowing to start or stop based on time elapsed
+ */
+UCLASS(Abstract)
+class PROPERTYANIMATORCORE_API UPropertyAnimatorCoreRangeTimeSource : public UPropertyAnimatorCoreTimeSourceBase
+{
+	GENERATED_BODY()
+
+public:
+	UPropertyAnimatorCoreRangeTimeSource()
+		: UPropertyAnimatorCoreTimeSourceBase(NAME_None)
+	{}
+
+	UPropertyAnimatorCoreRangeTimeSource(const FName& InSourceName)
+		: UPropertyAnimatorCoreTimeSourceBase(InSourceName)
+	{}
+
+	void SetUseStartTime(bool bInUse);
+
+	bool GetUseStartTime() const
+	{
+		return bUseStartTime;
+	}
+
+	void SetStartTime(double InStartTime);
+
+	double GetStartTime() const
+	{
+		return StartTime;
+	}
+
+	void SetUseStopTime(bool bInUse);
+
+	bool GetUseStopTime() const
+	{
+		return bUseStopTime;
+	}
+
+	void SetStopTime(double InStopTime);
+
+	double GetStopTime() const
+	{
+		return StopTime;
+	}
+
+	//~ Begin UPropertyAnimatorTimeSourceBase
+	virtual bool IsValidTimeElapsed(double InTimeElapsed) const override;
+	//~ End UPropertyAnimatorTimeSourceBase
+
+protected:
+	UPROPERTY(EditInstanceOnly, BlueprintReadWrite, Setter="SetUseStartTime", Getter="GetUseStartTime", Category="Animator")
+	bool bUseStartTime = false;
+
+	UPROPERTY(EditInstanceOnly, BlueprintReadWrite, Setter, Getter, Category="Animator", meta=(EditCondition="bUseStartTime", EditConditionHides))
+	double StartTime = 0;
+
+	UPROPERTY(EditInstanceOnly, BlueprintReadWrite, Setter="SetUseStopTime", Getter="GetUseStopTime", Category="Animator")
+	bool bUseStopTime = false;
+
+	UPROPERTY(EditInstanceOnly, BlueprintReadWrite, Setter, Getter, Category="Animator", meta=(EditCondition="bUseStopTime", EditConditionHides))
+	double StopTime = 0;
+};
\ No newline at end of file
diff --git a/Engine/Plugins/Experimental/PropertyAnimatorCore/Source/PropertyAnimatorCore/Public/TimeSources/PropertyAnimatorCoreTimeSourceBase.h b/Engine/Plugins/Experimental/PropertyAnimatorCore/Source/PropertyAnimatorCore/Public/TimeSources/PropertyAnimatorCoreTimeSourceBase.h
new file mode 100644
--- /dev/null
+++ b/Engine/Plugins/Experimental/PropertyAnimatorCore/Source/PropertyAnimatorCore/Public/TimeSources/PropertyAnimatorCoreTimeSourceBase.h
@@ -0,0 +1,88 @@
+// Copyright Epic Games, Inc. All Rights Reserved.
+
+#pragma once
+
+#include "PropertyAnimatorCoreTimeSourceBase.generated.h"
+
+class UPropertyAnimatorCoreBase;
+
+/**
+ * Abstract base class for time source used by property animators
+ * Can be transient or saved to disk if contains user set data
+ */
+UCLASS(Abstract)
+class PROPERTYANIMATORCORE_API UPropertyAnimatorCoreTimeSourceBase : public UObject
+{
+	GENERATED_BODY()
+
+	friend class UPropertyAnimatorCoreSubsystem;
+
+public:
+	UPropertyAnimatorCoreTimeSourceBase()
+		: UPropertyAnimatorCoreTimeSourceBase(NAME_None)
+	{}
+
+	UPropertyAnimatorCoreTimeSourceBase(const FName& InSourceName)
+		: TimeSourceName(InSourceName)
+	{}
+
+	void ActivateTimeSource();
+	void DeactivateTimeSource();
+
+	bool IsTimeSourceActive() const
+	{
+		return bTimeSourceActive;
+	}
+
+	TOptional<double> GetConditionalTimeElapsed();
+
+	/** Get the animator, this time source is on */
+	UPropertyAnimatorCoreBase* GetAnimator() const;
+
+	FName GetTimeSourceName() const
+	{
+		return TimeSourceName;
+	}
+
+protected:
+	/** Returns the time elapsed for animators */
+	virtual double GetTimeElapsed()
+	{
+		return 0;
+	}
+
+	/** Checks if this time source is ready to be used by the animator */
+	virtual bool IsTimeSourceReady() const
+	{
+		return false;
+	}
+
+	/** Check if the time elapsed is valid based on the context */
+	virtual bool IsValidTimeElapsed(double InTimeElapsed) const
+	{
+		return true;
+	}
+
+	/** Time source CDO is registered by subsystem */
+	virtual void OnTimeSourceRegistered() {}
+
+	/** Time source CDO is unregistered by subsystem */
+	virtual void OnTimeSourceUnregistered() {}
+
+	/** Time source is active on the animator */
+	virtual void OnTimeSourceActive() {}
+
+	/** Time source is inactive on the animator */
+	virtual void OnTimeSourceInactive() {}
+
+private:
+	/** Only used to display time elapsed */
+	UPROPERTY(Transient, VisibleInstanceOnly, Category="Animator", meta=(NoResetToDefault))
+	double TimeElapsed = 0;
+
+	/** Name used to display this time source to the user */
+	UPROPERTY(Transient)
+	FName TimeSourceName;
+
+	bool bTimeSourceActive = false;
+};
\ No newline at end of file
diff --git a/Engine/Plugins/Experimental/PropertyAnimatorCore/Source/PropertyAnimatorCore/Public/TimeSources/PropertyAnimatorCoreWorldTimeSource.h b/Engine/Plugins/Experimental/PropertyAnimatorCore/Source/PropertyAnimatorCore/Public/TimeSources/PropertyAnimatorCoreWorldTimeSource.h
new file mode 100644
--- /dev/null
+++ b/Engine/Plugins/Experimental/PropertyAnimatorCore/Source/PropertyAnimatorCore/Public/TimeSources/PropertyAnimatorCoreWorldTimeSource.h
@@ -0,0 +1,26 @@
+// Copyright Epic Games, Inc. All Rights Reserved.
+
+#pragma once
+
+#include "PropertyAnimatorCoreTimeSourceBase.h"
+#include "PropertyAnimatorCoreWorldTimeSource.generated.h"
+
+/**
+ * Time source that follows world time,
+ * Transient because it does not contain user saved data, can be recreated
+ */
+UCLASS()
+class PROPERTYANIMATORCORE_API UPropertyAnimatorCoreWorldTimeSource : public UPropertyAnimatorCoreTimeSourceBase
+{
+	GENERATED_BODY()
+
+public:
+	UPropertyAnimatorCoreWorldTimeSource()
+		: UPropertyAnimatorCoreTimeSourceBase(TEXT("World"))
+	{}
+
+	//~ Begin UPropertyAnimatorTimeSourceBase
+	virtual double GetTimeElapsed() override;
+	virtual bool IsTimeSourceReady() const override;
+	//~ End UPropertyAnimatorTimeSourceBase
+};
\ No newline at end of file
diff --git a/Engine/Plugins/Experimental/PropertyAnimatorCore/Source/PropertyAnimatorCoreEditor/Private/Customizations/PropertyAnimatorCoreEditorContextTypeCustomization.cpp b/Engine/Plugins/Experimental/PropertyAnimatorCore/Source/PropertyAnimatorCoreEditor/Private/Customizations/PropertyAnimatorCoreEditorContextTypeCustomization.cpp
new file mode 100644
--- /dev/null
+++ b/Engine/Plugins/Experimental/PropertyAnimatorCore/Source/PropertyAnimatorCoreEditor/Private/Customizations/PropertyAnimatorCoreEditorContextTypeCustomization.cpp
@@ -0,0 +1,199 @@
+// Copyright Epic Games, Inc. All Rights Reserved.
+
+#include "Customizations/PropertyAnimatorCoreEditorContextTypeCustomization.h"
+
+#include "DetailWidgetRow.h"
+#include "IDetailChildrenBuilder.h"
+#include "IPropertyUtilities.h"
+#include "Subsystems/PropertyAnimatorCoreSubsystem.h"
+#include "Widgets/Images/SImage.h"
+#include "Widgets/Input/SButton.h"
+#include "Widgets/Input/SCheckBox.h"
+#include "Widgets/SBoxPanel.h"
+
+#define LOCTEXT_NAMESPACE "PropertyAnimatorCoreEditorContextTypeCustomization"
+
+void FPropertyAnimatorCoreEditorContextTypeCustomization::CustomizeHeader(TSharedRef<IPropertyHandle> InPropertyHandle, FDetailWidgetRow& InRow, IPropertyTypeCustomizationUtils& InUtils)
+{
+	if (!InPropertyHandle->IsValidHandle())
+	{
+		return;
+	}
+
+	if (InPropertyHandle->GetNumOuterObjects() != 1)
+	{
+		return;
+	}
+
+	const UPropertyAnimatorCoreContext* PropertyContext = GetPropertyContextValue(InPropertyHandle);
+
+	if (!PropertyContext)
+	{
+		return;
+	}
+
+	PropertyContextHandle = InPropertyHandle;
+
+	const FName PropertyName = PropertyContext->GetAnimatedProperty().GetPropertyDisplayName();
+
+	InRow
+		.NameContent()
+		[
+			SNew(SHorizontalBox)
+			+ SHorizontalBox::Slot()
+			.AutoWidth()
+			[
+				SNew(SCheckBox)
+				.IsChecked(this, &FPropertyAnimatorCoreEditorContextTypeCustomization::IsPropertyEnabled)
+				.OnCheckStateChanged(this, &FPropertyAnimatorCoreEditorContextTypeCustomization::OnPropertyEnabled)
+			]
+			+ SHorizontalBox::Slot()
+			.FillWidth(1.f)
+			[
+				InPropertyHandle->CreatePropertyNameWidget(FText::FromName(PropertyName))
+			]
+		]
+		.ValueContent()
+		.HAlign(HAlign_Left)
+		[
+			SNew(SHorizontalBox)
+			+ SHorizontalBox::Slot()
+			.AutoWidth()
+			.Padding(2.f, 0.f)
+			[
+				SNew(SButton)
+				.ContentPadding(2.f)
+				.ToolTipText(LOCTEXT("UnlinkProperty", "Unlink property from animator"))
+				.OnClicked(this, &FPropertyAnimatorCoreEditorContextTypeCustomization::UnlinkProperty)
+				.Content()
+				[
+					SNew(SImage)
+					.Image(FAppStyle::Get().GetBrush("Icons.Delete"))
+					.DesiredSizeOverride(FVector2D(16.f))
+				]
+			]
+		];
+}
+
+void FPropertyAnimatorCoreEditorContextTypeCustomization::CustomizeChildren(TSharedRef<IPropertyHandle> InPropertyHandle, IDetailChildrenBuilder& InBuilder, IPropertyTypeCustomizationUtils& InUtils)
+{
+	static const TSet<FName> SkipProperties
+	{
+		GET_MEMBER_NAME_CHECKED(UPropertyAnimatorCoreContext, bAnimated)
+	};
+
+	UObject* PropertyObject;
+	InPropertyHandle->GetValue(PropertyObject);
+
+	if (UPropertyAnimatorCoreContext* Options = Cast<UPropertyAnimatorCoreContext>(PropertyObject))
+	{
+		for (const FProperty* Property : TFieldRange<FProperty>(Options->GetClass()))
+		{
+			if (Property
+				&& Property->HasAnyPropertyFlags(EPropertyFlags::CPF_Edit)
+				&& !SkipProperties.Contains(Property->GetFName()))
+			{
+				IDetailPropertyRow* Row = InBuilder.AddExternalObjectProperty({Options}, Property->GetFName());
+
+				if (Row && Property->GetFName() == GET_MEMBER_NAME_CHECKED(UPropertyAnimatorCoreContext, ConverterRule))
+				{
+					if (const TSharedPtr<IPropertyHandle> ConverterHandle = Row->GetPropertyHandle())
+					{
+						PropertyUtilities = InUtils.GetPropertyUtilities();
+
+						// Fix for container properties (array, set, map) in instanced struct not updating details view (when adding/removing entries)
+						ConverterHandle->SetOnChildPropertyValueChangedWithData(TDelegate<void(const FPropertyChangedEvent&)>::CreateLambda([this, Row](const FPropertyChangedEvent& InEvent)
+						{
+							if (PropertyUtilities.IsValid()
+								&& InEvent.Property
+								&& (InEvent.Property->IsA<FArrayProperty>()
+								|| InEvent.Property->IsA<FSetProperty>()
+								|| InEvent.Property->IsA<FMapProperty>()))
+							{
+								PropertyUtilities->RequestRefresh();
+							}
+						}));
+
+						// Hide customization dropdown to avoid user changing instanced struct type
+						PropertyUtilities->EnqueueDeferredAction(FSimpleDelegate::CreateLambda([Row]()
+						{
+							if (Row && Row->CustomValueWidget())
+							{
+								Row->CustomValueWidget()->Widget->SetVisibility(EVisibility::Hidden);
+							}
+						}));
+					}
+				}
+			}
+		}
+	}
+}
+
+UPropertyAnimatorCoreContext* FPropertyAnimatorCoreEditorContextTypeCustomization::GetPropertyContextValue(TSharedPtr<IPropertyHandle> InHandle)
+{
+	if (!InHandle.IsValid() || !InHandle->IsValidHandle())
+	{
+		return nullptr;
+	}
+
+	UObject* ValueData = nullptr;
+
+	if (InHandle->GetValue(ValueData) == FPropertyAccess::Result::Success)
+	{
+		return Cast<UPropertyAnimatorCoreContext>(ValueData);
+	}
+
+	return nullptr;
+}
+
+ECheckBoxState FPropertyAnimatorCoreEditorContextTypeCustomization::IsPropertyEnabled() const
+{
+	if (const UPropertyAnimatorCoreContext* PropertyContext = GetPropertyContextValue(PropertyContextHandle))
+	{
+		return PropertyContext->IsAnimated() ? ECheckBoxState::Checked : ECheckBoxState::Unchecked;
+	}
+
+	return ECheckBoxState::Undetermined;
+}
+
+void FPropertyAnimatorCoreEditorContextTypeCustomization::OnPropertyEnabled(ECheckBoxState InNewState) const
+{
+	if (UPropertyAnimatorCoreContext* PropertyContext = GetPropertyContextValue(PropertyContextHandle))
+	{
+		PropertyContext->SetAnimated(InNewState == ECheckBoxState::Checked);
+	}
+}
+
+FReply FPropertyAnimatorCoreEditorContextTypeCustomization::UnlinkProperty() const
+{
+	if (!PropertyContextHandle.IsValid())
+	{
+		return FReply::Unhandled();
+	}
+
+	const UPropertyAnimatorCoreContext* PropertyContext = GetPropertyContextValue(PropertyContextHandle);
+	if (!PropertyContext)
+	{
+		return FReply::Unhandled();
+	}
+
+	TArray<UObject*> Owners;
+	PropertyContextHandle->GetOuterObjects(Owners);
+
+	if (UPropertyAnimatorCoreSubsystem* AnimatorSubsystem = UPropertyAnimatorCoreSubsystem::Get())
+	{
+		FPropertyAnimatorCoreData PropertyData = PropertyContext->GetAnimatedProperty();
+
+		for (UObject* Owner : Owners)
+		{
+			if (UPropertyAnimatorCoreBase* PropertyAnimator = Cast<UPropertyAnimatorCoreBase>(Owner))
+			{
+				AnimatorSubsystem->UnlinkAnimatorProperty(PropertyAnimator, PropertyData, true);
+			}
+		}
+	}
+
+	return FReply::Handled();
+}
+
+#undef LOCTEXT_NAMESPACE
