[PropertyAnimator] Moved plugin

diff --git a/Engine/Plugins/Experimental/PropertyAnimatorCore/Source/PropertyAnimatorCoreEditor/Private/Customizations/PropertyAnimatorCoreEditorStackCustomization.cpp b/Engine/Plugins/Experimental/PropertyAnimatorCore/Source/PropertyAnimatorCoreEditor/Private/Customizations/PropertyAnimatorCoreEditorStackCustomization.cpp
new file mode 100644
--- /dev/null
+++ b/Engine/Plugins/Experimental/PropertyAnimatorCore/Source/PropertyAnimatorCoreEditor/Private/Customizations/PropertyAnimatorCoreEditorStackCustomization.cpp
@@ -0,0 +1,453 @@
+// Copyright Epic Games, Inc. All Rights Reserved.
+
+#include "PropertyAnimatorCoreEditorStackCustomization.h"
+
+#include "Components/PropertyAnimatorCoreComponent.h"
+#include "Contexts/OperatorStackEditorMenuContext.h"
+#include "Animators/PropertyAnimatorCoreBase.h"
+#include "Framework/Commands/GenericCommands.h"
+#include "Items/OperatorStackEditorObjectItem.h"
+#include "Menus/PropertyAnimatorCoreEditorMenu.h"
+#include "Styling/SlateIconFinder.h"
+#include "Subsystems/PropertyAnimatorCoreEditorSubsystem.h"
+#include "Subsystems/PropertyAnimatorCoreSubsystem.h"
+#include "ToolMenus.h"
+#include "ToolMenu.h"
+#include "Widgets/PropertyAnimatorCoreEditorEditPanelOptions.h"
+
+#define LOCTEXT_NAMESPACE "PropertyAnimatorEditorStackCustomization"
+
+UPropertyAnimatorCoreEditorStackCustomization::UPropertyAnimatorCoreEditorStackCustomization()
+	: UOperatorStackEditorStackCustomization(
+		TEXT("Animators")
+		, LOCTEXT("CustomizationLabel", "Animators")
+		, 0
+	)
+{
+	RegisterCustomizationFor(UPropertyAnimatorCoreBase::StaticClass());
+	RegisterCustomizationFor(UPropertyAnimatorCoreComponent::StaticClass());
+}
+
+bool UPropertyAnimatorCoreEditorStackCustomization::TransformContextItem(const FOperatorStackEditorItemPtr& InItem, TArray<FOperatorStackEditorItemPtr>& OutTransformedItems) const
+{
+	// If we have an actor then display controllers
+	if (InItem->IsA<AActor>())
+	{
+		const AActor* Actor = InItem->Get<AActor>();
+
+		if (UActorComponent* ControllerComponent = Actor->FindComponentByClass(UPropertyAnimatorCoreComponent::StaticClass()))
+		{
+			OutTransformedItems.Add(MakeShared<FOperatorStackEditorObjectItem>(ControllerComponent));
+		}
+
+		return true;
+	}
+	// If we have a property controller component then display stack
+	else if (InItem->IsA<UPropertyAnimatorCoreComponent>())
+	{
+		const UPropertyAnimatorCoreComponent* ControllerComponent = InItem->Get<UPropertyAnimatorCoreComponent>();
+
+		for (UPropertyAnimatorCoreBase* Controller : ControllerComponent->GetAnimators())
+		{
+			if (Controller && Controller->IsA<UPropertyAnimatorCoreBase>())
+			{
+				OutTransformedItems.Add(MakeShared<FOperatorStackEditorObjectItem>(Controller));
+			}
+		}
+
+		return true;
+	}
+
+	return Super::TransformContextItem(InItem, OutTransformedItems);
+}
+
+void UPropertyAnimatorCoreEditorStackCustomization::CustomizeStackHeader(const FOperatorStackEditorTree& InItemTree, FOperatorStackEditorHeaderBuilder& InHeaderBuilder)
+{
+	if (!InItemTree.GetContext().GetItems().IsEmpty())
+	{
+		static const FName AddAnimatorMenuName = TEXT("AddAnimatorMenu");
+		if (!UToolMenus::Get()->IsMenuRegistered(AddAnimatorMenuName))
+		{
+			UToolMenu* const AddControllerMenu = UToolMenus::Get()->RegisterMenu(AddAnimatorMenuName, NAME_None, EMultiBoxType::Menu);
+			AddControllerMenu->AddDynamicSection(TEXT("FillAddAnimatorMenuSection"), FNewToolMenuDelegate::CreateUObject(this, &UPropertyAnimatorCoreEditorStackCustomization::FillAddAnimatorMenuSection));
+		}
+
+		// Pinned search keywords
+		TSet<FString> PinnedControllerNames;
+		for (const FOperatorStackEditorItemPtr& SupportedItem : InItemTree.GetAllItems())
+		{
+			if (const UPropertyAnimatorCoreBase* Controller = SupportedItem->Get<UPropertyAnimatorCoreBase>())
+			{
+				PinnedControllerNames.Add(Controller->GetAnimatorOriginalName().ToString());
+			}
+		}
+
+		InHeaderBuilder
+			.SetToolMenu(
+				AddAnimatorMenuName
+				, LOCTEXT("AddControllersMenu", "Add Animators")
+				, FAppStyle::GetBrush(TEXT("Icons.PlusCircle")))
+			.SetSearchAllowed(true)
+			.SetSearchPinnedKeywords(PinnedControllerNames);
+	}
+
+	Super::CustomizeStackHeader(InItemTree, InHeaderBuilder);
+}
+
+void UPropertyAnimatorCoreEditorStackCustomization::CustomizeItemHeader(const FOperatorStackEditorItemPtr& InItem, const FOperatorStackEditorTree& InItemTree, FOperatorStackEditorHeaderBuilder& InHeaderBuilder)
+{
+	// Customize component header
+	if (InItem->IsA<UPropertyAnimatorCoreComponent>())
+	{
+		FBoolProperty* EnableProperty = FindFProperty<FBoolProperty>(UPropertyAnimatorCoreComponent::StaticClass(), GET_MEMBER_NAME_CHECKED(UPropertyAnimatorCoreComponent, bAnimatorsEnabled));
+
+		const FSlateIcon ClassIcon = FSlateIconFinder::FindIconForClass(UPropertyAnimatorCoreComponent::StaticClass());
+
+		// Action menu available in header in slim toolbar
+		static const FName HeaderComponentMenuName = TEXT("HeaderComponentMenu");
+		if (!UToolMenus::Get()->IsMenuRegistered(HeaderComponentMenuName))
+		{
+			UToolMenu* const HeaderAnimatorMenu = UToolMenus::Get()->RegisterMenu(HeaderComponentMenuName, NAME_None, EMultiBoxType::SlimHorizontalToolBar);
+			HeaderAnimatorMenu->AddDynamicSection(TEXT("FillHeaderComponentMenu"), FNewToolMenuDelegate::CreateUObject(this, &UPropertyAnimatorCoreEditorStackCustomization::FillComponentHeaderActionMenu));
+		}
+
+		FString HeaderLabel = TEXT("Animators");
+
+		// Add actor name next to label when we customize multiple items
+		if (InItemTree.GetRootItems().Num() > 1)
+		{
+			if (const UPropertyAnimatorCoreComponent* ControllerComponent = InItem->Get<UPropertyAnimatorCoreComponent>())
+			{
+				HeaderLabel += + TEXT(" (") + ControllerComponent->GetOwner()->GetActorNameOrLabel() + TEXT(")");
+			}
+		}
+
+		InHeaderBuilder
+			.SetProperty(EnableProperty)
+			.SetToolbarMenu(HeaderComponentMenuName)
+			.SetIcon(ClassIcon.GetIcon())
+			.SetLabel(FText::FromString(HeaderLabel));
+	}
+	// Customize controller header
+	else if (InItem->IsA<UPropertyAnimatorCoreBase>())
+	{
+		UPropertyAnimatorCoreBase* Animator = InItem->Get<UPropertyAnimatorCoreBase>();
+
+		FBoolProperty* EnableProperty = FindFProperty<FBoolProperty>(UPropertyAnimatorCoreBase::StaticClass(), GET_MEMBER_NAME_CHECKED(UPropertyAnimatorCoreBase, bAnimatorEnabled));
+
+		const FSlateIcon ClassIcon = FSlateIconFinder::FindIconForClass(Animator->GetClass());
+
+		// Commands for item on key events
+		const TSharedPtr<FUICommandList> Commands = CreateAnimatorCommands(Animator);
+
+		// Action menu available in header in slim toolbar
+		static const FName HeaderModifierMenuName = TEXT("HeaderControllerMenu");
+		if (!UToolMenus::Get()->IsMenuRegistered(HeaderModifierMenuName))
+		{
+			UToolMenu* const HeaderModifierMenu = UToolMenus::Get()->RegisterMenu(HeaderModifierMenuName, NAME_None, EMultiBoxType::SlimHorizontalToolBar);
+			HeaderModifierMenu->AddDynamicSection(TEXT("FillHeaderControllerMenu"), FNewToolMenuDelegate::CreateUObject(this, &UPropertyAnimatorCoreEditorStackCustomization::FillAnimatorHeaderActionMenu));
+		}
+
+		// Context menu available when right clicking on item
+		static const FName ContextModifierMenuName = TEXT("ContextControllerMenu");
+		if (!UToolMenus::Get()->IsMenuRegistered(ContextModifierMenuName))
+		{
+			UToolMenu* const ContextModifierMenu = UToolMenus::Get()->RegisterMenu(ContextModifierMenuName, NAME_None, EMultiBoxType::Menu);
+			ContextModifierMenu->AddDynamicSection(TEXT("FillContextControllerMenu"), FNewToolMenuDelegate::CreateUObject(this, &UPropertyAnimatorCoreEditorStackCustomization::FillAnimatorContextActionMenu));
+		}
+
+		const TSet<FString> SearchKeywords
+		{
+			Animator->GetAnimatorOriginalName().ToString(),
+			Animator->GetAnimatorDisplayName()
+		};
+
+		/** Show last execution error messages if failed execution */
+		TAttribute<EOperatorStackEditorMessageType> MessageType = EOperatorStackEditorMessageType::None;
+		TAttribute<FText> MessageText = FText::GetEmpty();
+
+		TWeakObjectPtr<UPropertyAnimatorCoreBase> AnimatorWeak(Animator);
+
+		MessageType = TAttribute<EOperatorStackEditorMessageType>::CreateLambda([AnimatorWeak]()
+		{
+			if (const UPropertyAnimatorCoreBase* Animator = AnimatorWeak.Get())
+			{
+				return Animator->GetLinkedPropertiesCount() > 0
+					? EOperatorStackEditorMessageType::None
+					: EOperatorStackEditorMessageType::Info;
+			}
+
+			return EOperatorStackEditorMessageType::None;
+		});
+
+		MessageText = TAttribute<FText>::CreateLambda([AnimatorWeak]()
+		{
+			if (AnimatorWeak.IsValid())
+			{
+				return LOCTEXT("NoPropertiesLinked", "No properties are currently linked to this animator");
+			}
+
+			return FText::GetEmpty();
+		});
+
+		static const FLinearColor AnimatorColor = FLinearColor(FColor::Orange).Desaturate(0.25);
+
+		InHeaderBuilder
+			.SetBorderColor(AnimatorColor)
+			.SetSearchAllowed(true)
+			.SetSearchKeywords(SearchKeywords)
+			.SetExpandable(true)
+			.SetIcon(ClassIcon.GetIcon())
+			.SetLabel(FText::FromString(Animator->GetAnimatorDisplayName()))
+			.SetProperty(EnableProperty)
+			.SetCommandList(Commands)
+			.SetToolbarMenu(HeaderModifierMenuName)
+			.SetContextMenu(ContextModifierMenuName)
+			.SetMessageBox(MessageType, MessageText);
+	}
+
+	Super::CustomizeItemHeader(InItem, InItemTree, InHeaderBuilder);
+}
+
+void UPropertyAnimatorCoreEditorStackCustomization::CustomizeItemBody(const FOperatorStackEditorItemPtr& InItem, const FOperatorStackEditorTree& InItemTree, FOperatorStackEditorBodyBuilder& InBodyBuilder)
+{
+	// Customize controller body
+	if (InItem->IsA<UPropertyAnimatorCoreBase>())
+	{
+		FBoolProperty* EnableProperty = FindFProperty<FBoolProperty>(UPropertyAnimatorCoreBase::StaticClass(), GET_MEMBER_NAME_CHECKED(UPropertyAnimatorCoreBase, bAnimatorEnabled));
+
+		InBodyBuilder
+			.SetShowDetailsView(true)
+			.DisallowProperty(EnableProperty);
+	}
+
+	Super::CustomizeItemBody(InItem, InItemTree, InBodyBuilder);
+}
+
+void UPropertyAnimatorCoreEditorStackCustomization::CustomizeItemFooter(const FOperatorStackEditorItemPtr& InItem, const FOperatorStackEditorTree& InItemTree, FOperatorStackEditorFooterBuilder& InFooterBuilder)
+{
+	// Customize component footer
+	if (InItem->IsA<UPropertyAnimatorCoreComponent>())
+	{
+		FProperty* MagnitudeProperty = FindFProperty<FProperty>(UPropertyAnimatorCoreComponent::StaticClass(), GET_MEMBER_NAME_CHECKED(UPropertyAnimatorCoreComponent, AnimatorsMagnitude));
+
+		InFooterBuilder
+			.AllowProperty(MagnitudeProperty)
+			.SetShowDetailsView(true);
+	}
+
+	Super::CustomizeItemFooter(InItem, InItemTree, InFooterBuilder);
+}
+
+bool UPropertyAnimatorCoreEditorStackCustomization::OnIsItemDraggable(const FOperatorStackEditorItemPtr& InItem)
+{
+	// Allow selection of item in listview for commands to work
+	if (InItem->IsA<UPropertyAnimatorCoreBase>())
+	{
+		return true;
+	}
+
+	return Super::OnIsItemDraggable(InItem);
+}
+
+const FSlateBrush* UPropertyAnimatorCoreEditorStackCustomization::GetIcon() const
+{
+	return FSlateIconFinder::FindIconForClass(UPropertyAnimatorCoreBase::StaticClass()).GetIcon();
+}
+
+void UPropertyAnimatorCoreEditorStackCustomization::RemoveControllerAction(UPropertyAnimatorCoreBase* InAnimator) const
+{
+	if (!InAnimator)
+	{
+		return;
+	}
+
+	const UPropertyAnimatorCoreSubsystem* Subsystem = UPropertyAnimatorCoreSubsystem::Get();
+
+	if (!Subsystem)
+	{
+		return;
+	}
+
+	Subsystem->RemoveAnimator(InAnimator, true);
+}
+
+void UPropertyAnimatorCoreEditorStackCustomization::RemoveControllersAction(UPropertyAnimatorCoreComponent* InComponent) const
+{
+	const UPropertyAnimatorCoreSubsystem* Subsystem = UPropertyAnimatorCoreSubsystem::Get();
+	if (!InComponent || !Subsystem)
+	{
+		return;
+	}
+
+	TSet<UPropertyAnimatorCoreBase*> Animators;
+	for (const TObjectPtr<UPropertyAnimatorCoreBase>& Animator : InComponent->GetAnimators())
+	{
+		Animators.Add(Animator);
+	}
+
+	Subsystem->RemoveAnimators(Animators, true);
+}
+
+void UPropertyAnimatorCoreEditorStackCustomization::FillAddAnimatorMenuSection(UToolMenu* InToolMenu) const
+{
+	if (!InToolMenu)
+	{
+		return;
+	}
+
+	const UOperatorStackEditorMenuContext* const AddAnimatorContext = InToolMenu->FindContext<UOperatorStackEditorMenuContext>();
+	if (!AddAnimatorContext)
+	{
+		return;
+	}
+
+	const FOperatorStackEditorContextPtr Context = AddAnimatorContext->GetContext();
+	if (!Context)
+	{
+		return;
+	}
+
+	UPropertyAnimatorCoreEditorSubsystem* AnimatorEditorSubsystem = UPropertyAnimatorCoreEditorSubsystem::Get();
+	if (!AnimatorEditorSubsystem)
+	{
+		return;
+	}
+
+	TSet<UObject*> ContextObjects;
+	for (const FOperatorStackEditorItemPtr& ContextItem : Context->GetItems())
+	{
+		if (ContextItem->IsA<UObject>())
+		{
+			ContextObjects.Add(ContextItem->Get<UObject>());
+		}
+	}
+
+	const FPropertyAnimatorCoreEditorMenuContext MenuContext(ContextObjects, {});
+	const FPropertyAnimatorCoreEditorMenuOptions MenuOptions({EPropertyAnimatorCoreEditorMenuType::New});
+	AnimatorEditorSubsystem->FillAnimatorMenu(InToolMenu, MenuContext, MenuOptions);
+}
+
+void UPropertyAnimatorCoreEditorStackCustomization::FillComponentHeaderActionMenu(UToolMenu* InToolMenu)
+{
+	if (!InToolMenu)
+	{
+		return;
+	}
+
+	const UOperatorStackEditorMenuContext* const MenuContext = InToolMenu->FindContext<UOperatorStackEditorMenuContext>();
+	if (!MenuContext)
+	{
+		return;
+	}
+
+	const FOperatorStackEditorItemPtr ItemContext = MenuContext->GetItem();
+	if (!ItemContext)
+	{
+		return;
+	}
+
+	UPropertyAnimatorCoreComponent* ControllerComponent = ItemContext->Get<UPropertyAnimatorCoreComponent>();
+	if (!IsValid(ControllerComponent))
+	{
+		return;
+	}
+
+	// Add remove controllers entry
+	const FToolMenuEntry RemoveControllersAction = FToolMenuEntry::InitToolBarButton(
+		TEXT("RemoveControllersMenuEntry")
+		, FExecuteAction::CreateUObject(this, &UPropertyAnimatorCoreEditorStackCustomization::RemoveControllersAction, ControllerComponent)
+		, FText::GetEmpty()
+		, FText::GetEmpty()
+		, FSlateIcon(FAppStyle::GetAppStyleSetName(), "Icons.Delete")
+	);
+
+	InToolMenu->AddMenuEntry(RemoveControllersAction.Name, RemoveControllersAction);
+}
+
+void UPropertyAnimatorCoreEditorStackCustomization::FillAnimatorHeaderActionMenu(UToolMenu* InToolMenu)
+{
+	if (!InToolMenu)
+	{
+		return;
+	}
+
+	const UOperatorStackEditorMenuContext* const MenuContext = InToolMenu->FindContext<UOperatorStackEditorMenuContext>();
+	if (!MenuContext)
+	{
+		return;
+	}
+
+	const FOperatorStackEditorItemPtr ItemContext = MenuContext->GetItem();
+	if (!ItemContext)
+	{
+		return;
+	}
+
+	UPropertyAnimatorCoreBase* Controller = ItemContext->Get<UPropertyAnimatorCoreBase>();
+	if (!IsValid(Controller))
+	{
+		return;
+	}
+
+	// Add remove controller entry
+	const FToolMenuEntry RemoveControllerAction = FToolMenuEntry::InitToolBarButton(
+		TEXT("RemoveControllerMenuEntry")
+		, FExecuteAction::CreateUObject(this, &UPropertyAnimatorCoreEditorStackCustomization::RemoveControllerAction, Controller)
+		, FText::GetEmpty()
+		, FText::GetEmpty()
+		, FSlateIcon(FAppStyle::GetAppStyleSetName(), "Icons.Delete")
+	);
+
+	InToolMenu->AddMenuEntry(RemoveControllerAction.Name, RemoveControllerAction);
+}
+
+void UPropertyAnimatorCoreEditorStackCustomization::FillAnimatorContextActionMenu(UToolMenu* InToolMenu) const
+{
+	if (!InToolMenu)
+	{
+		return;
+	}
+
+	const UOperatorStackEditorMenuContext* const MenuContext = InToolMenu->FindContext<UOperatorStackEditorMenuContext>();
+	if (!MenuContext)
+	{
+		return;
+	}
+
+	const FOperatorStackEditorItemPtr ItemContext = MenuContext->GetItem();
+	if (!ItemContext)
+	{
+		return;
+	}
+
+	const UPropertyAnimatorCoreBase* Controller = ItemContext->Get<UPropertyAnimatorCoreBase>();
+	if (!IsValid(Controller))
+	{
+		return;
+	}
+
+	// Lets get the commands list to bind the menu entry and execute it when clicked
+	TSharedPtr<const FUICommandList> Commands;
+	const TSharedPtr<FUICommandInfo> DeleteCommand = FGenericCommands::Get().Delete;
+	InToolMenu->Context.GetActionForCommand(DeleteCommand, Commands);
+	const FToolMenuEntry RemoveControllerMenuEntry = FToolMenuEntry::InitMenuEntryWithCommandList(DeleteCommand, Commands);
+
+	InToolMenu->AddMenuEntry(RemoveControllerMenuEntry.Name, RemoveControllerMenuEntry);
+}
+
+TSharedRef<FUICommandList> UPropertyAnimatorCoreEditorStackCustomization::CreateAnimatorCommands(UPropertyAnimatorCoreBase* InAnimator)
+{
+	TSharedRef<FUICommandList> Commands = MakeShared<FUICommandList>();
+
+	Commands->MapAction(FGenericCommands::Get().Delete,
+		FUIAction(
+			FExecuteAction::CreateUObject(this, &UPropertyAnimatorCoreEditorStackCustomization::RemoveControllerAction, InAnimator),
+			FCanExecuteAction()
+		)
+	);
+
+	return Commands;
+}
+
+#undef LOCTEXT_NAMESPACE
diff --git a/Engine/Plugins/Experimental/PropertyAnimatorCore/Source/PropertyAnimatorCoreEditor/Private/Customizations/PropertyAnimatorCoreEditorStackCustomization.h b/Engine/Plugins/Experimental/PropertyAnimatorCore/Source/PropertyAnimatorCoreEditor/Private/Customizations/PropertyAnimatorCoreEditorStackCustomization.h
new file mode 100644
--- /dev/null
+++ b/Engine/Plugins/Experimental/PropertyAnimatorCore/Source/PropertyAnimatorCoreEditor/Private/Customizations/PropertyAnimatorCoreEditorStackCustomization.h
@@ -0,0 +1,45 @@
+// Copyright Epic Games, Inc. All Rights Reserved.
+
+#pragma once
+
+#include "Animators/PropertyAnimatorCoreBase.h"
+#include "Customizations/OperatorStackEditorStackCustomization.h"
+#include "PropertyAnimatorCoreEditorStackCustomization.generated.h"
+
+class UPropertyAnimatorCoreComponent;
+class UPropertyAnimatorCoreBase;
+class UToolMenu;
+
+/** Property Controller customization for operator stack tab */
+UCLASS()
+class UPropertyAnimatorCoreEditorStackCustomization : public UOperatorStackEditorStackCustomization
+{
+	GENERATED_BODY()
+
+public:
+	UPropertyAnimatorCoreEditorStackCustomization();
+
+	//~ Begin UOperatorStackEditorStackCustomization
+	virtual bool TransformContextItem(const FOperatorStackEditorItemPtr& InItem, TArray<FOperatorStackEditorItemPtr>& OutTransformedItems) const override;
+	virtual void CustomizeStackHeader(const FOperatorStackEditorTree& InItemTree, FOperatorStackEditorHeaderBuilder& InHeaderBuilder) override;
+	virtual void CustomizeItemHeader(const FOperatorStackEditorItemPtr& InItem, const FOperatorStackEditorTree& InItemTree, FOperatorStackEditorHeaderBuilder& InHeaderBuilder) override;
+	virtual void CustomizeItemBody(const FOperatorStackEditorItemPtr& InItem, const FOperatorStackEditorTree& InItemTree, FOperatorStackEditorBodyBuilder& InBodyBuilder) override;
+	virtual void CustomizeItemFooter(const FOperatorStackEditorItemPtr& InItem, const FOperatorStackEditorTree& InItemTree, FOperatorStackEditorFooterBuilder& InFooterBuilder) override;
+	virtual bool OnIsItemDraggable(const FOperatorStackEditorItemPtr& InItem) override;
+	virtual const FSlateBrush* GetIcon() const override;
+	//~ End UOperatorStackEditorStackCustomization
+
+protected:
+	/** Remove animator menu action */
+	void RemoveControllerAction(UPropertyAnimatorCoreBase* InAnimator) const;
+	void RemoveControllersAction(UPropertyAnimatorCoreComponent* InComponent) const;
+
+	/** Fill item action menus */
+	void FillAddAnimatorMenuSection(UToolMenu* InToolMenu) const;
+	void FillComponentHeaderActionMenu(UToolMenu* InToolMenu);
+	void FillAnimatorHeaderActionMenu(UToolMenu* InToolMenu);
+	void FillAnimatorContextActionMenu(UToolMenu* InToolMenu) const;
+
+	/** Create commands for animator */
+	TSharedRef<FUICommandList> CreateAnimatorCommands(UPropertyAnimatorCoreBase* InAnimator);
+};
\ No newline at end of file
diff --git a/Engine/Plugins/Experimental/PropertyAnimatorCore/Source/PropertyAnimatorCoreEditor/Private/Menus/PropertyAnimatorCoreEditorMenu.cpp b/Engine/Plugins/Experimental/PropertyAnimatorCore/Source/PropertyAnimatorCoreEditor/Private/Menus/PropertyAnimatorCoreEditorMenu.cpp
new file mode 100644
--- /dev/null
+++ b/Engine/Plugins/Experimental/PropertyAnimatorCore/Source/PropertyAnimatorCoreEditor/Private/Menus/PropertyAnimatorCoreEditorMenu.cpp
@@ -0,0 +1,878 @@
+// Copyright Epic Games, Inc. All Rights Reserved.
+
+#include "Menus/PropertyAnimatorCoreEditorMenu.h"
+
+#include "Animators/PropertyAnimatorCoreBase.h"
+#include "Components/PropertyAnimatorCoreComponent.h"
+#include "Engine/World.h"
+#include "PropertyHandle.h"
+#include "Presets/PropertyAnimatorCorePresetBase.h"
+#include "Styling/SlateIconFinder.h"
+#include "Subsystems/PropertyAnimatorCoreEditorSubsystem.h"
+#include "Subsystems/PropertyAnimatorCoreSubsystem.h"
+#include "Widgets/PropertyAnimatorCoreEditorEditPanelOptions.h"
+
+#define LOCTEXT_NAMESPACE "PropertyAnimatorCoreEditorMenu"
+
+void UE::PropertyAnimatorCoreEditor::Menu::FillEditAnimatorSection(UToolMenu* InMenu, TSharedRef<FPropertyAnimatorCoreEditorMenuData> InMenuData)
+{
+	if (!InMenu || InMenuData->GetContext().IsEmpty())
+	{
+		return;
+	}
+
+	const UPropertyAnimatorCoreSubsystem* Subsystem = UPropertyAnimatorCoreSubsystem::Get();
+	if (!Subsystem)
+	{
+		return;
+	}
+
+	const TSet<AActor*> Actors = InMenuData->GetContext().GetActors();
+	if (Actors.Num() != 1)
+	{
+		return;
+	}
+
+	AActor* AnimatorActor = Actors.Array()[0];
+	if (!IsValid(AnimatorActor))
+	{
+		return;
+	}
+
+	FToolMenuSection& ManageAnimatorsSection = InMenu->FindOrAddSection(TEXT("ManageAnimators"), LOCTEXT("ManageAnimators.Label", "Manage Animators"));
+	const FSlateIcon AnimatorIcon = FSlateIconFinder::FindIconForClass(UPropertyAnimatorCoreBase::StaticClass());
+
+	ManageAnimatorsSection.AddMenuEntry(
+		TEXT("OpenEditAnimatorsWindow")
+	   , LOCTEXT("OpenEditAnimatorsWindow.Label", "Edit Animators")
+	   , LOCTEXT("OpenEditAnimatorsWindow.Tooltip", "Open edit animators window")
+	   , AnimatorIcon
+	   , FUIAction(
+		   FExecuteAction::CreateStatic(&ExecuteEditAnimatorAction, AnimatorActor)
+	   )
+	);
+}
+
+void UE::PropertyAnimatorCoreEditor::Menu::FillNewAnimatorSection(UToolMenu* InMenu, TSharedRef<FPropertyAnimatorCoreEditorMenuData> InMenuData)
+{
+	if (!InMenu || InMenuData->GetContext().IsEmpty())
+	{
+		return;
+	}
+
+	const UPropertyAnimatorCoreSubsystem* Subsystem = UPropertyAnimatorCoreSubsystem::Get();
+	if (!Subsystem)
+	{
+		return;
+	}
+
+	TSet<UPropertyAnimatorCoreBase*> NewAvailableAnimators = Subsystem->GetAvailableAnimators();
+	for (const FPropertyAnimatorCoreData& Property : InMenuData->GetContext().GetProperties())
+	{
+		NewAvailableAnimators = NewAvailableAnimators.Intersect(Subsystem->GetAvailableAnimators(&Property));
+	}
+
+	constexpr bool bCloseMenuAfterSelection = false;
+	constexpr bool bOpenOnClick = false;
+
+	FToolMenuSection& NewAnimatorsSection = InMenu->FindOrAddSection(TEXT("NewAnimators"), LOCTEXT("NewAnimators.Label", "New Animators"));
+
+	for (UPropertyAnimatorCoreBase* NewAnimator : NewAvailableAnimators)
+	{
+		const FName MenuName = NewAnimator->GetAnimatorOriginalName();
+		const FText MenuLabel = FText::FromName(MenuName);
+		const FSlateIcon MenuIcon = FSlateIconFinder::FindIconForClass(NewAnimator->GetClass());
+
+		NewAnimatorsSection.AddSubMenu(
+			MenuName
+			, MenuLabel
+			, LOCTEXT("NewAnimator.Tooltip", "Create a new animator")
+			, FNewToolMenuDelegate::CreateLambda(&FillNewAnimatorSubmenu, NewAnimator, InMenuData)
+			, bOpenOnClick
+			, MenuIcon
+			, bCloseMenuAfterSelection
+		);
+	}
+}
+
+void UE::PropertyAnimatorCoreEditor::Menu::FillExistingAnimatorSection(UToolMenu* InMenu, TSharedRef<FPropertyAnimatorCoreEditorMenuData> InMenuData)
+{
+	if (!InMenu || InMenuData->GetContext().IsEmpty())
+	{
+		return;
+	}
+
+	const UPropertyAnimatorCoreSubsystem* Subsystem = UPropertyAnimatorCoreSubsystem::Get();
+	if (!Subsystem)
+	{
+		return;
+	}
+
+	FToolMenuSection& ExistingAnimatorsSection = InMenu->FindOrAddSection(TEXT("ExistingAnimators"), LOCTEXT("ExistingAnimators.Label", "Existing Animators"));
+
+	constexpr bool bCloseMenuAfterSelection = true;
+	constexpr bool bOpenOnClick = false;
+
+	for (const FPropertyAnimatorCoreData& Property : InMenuData->GetContext().GetProperties())
+	{
+		for (UPropertyAnimatorCoreBase* Animator : Subsystem->GetExistingAnimators(Property))
+		{
+			const FName MenuName(Animator->GetAnimatorDisplayName());
+			const FText MenuLabel = FText::FromName(MenuName);
+			const FSlateIcon MenuIcon = FSlateIconFinder::FindIconForClass(Animator->GetClass());
+
+			ExistingAnimatorsSection.AddSubMenu(
+				MenuName
+				, MenuLabel
+				, LOCTEXT("ExistingAnimatorSection.Tooltip", "Link or unlink properties for this animator")
+				, FNewToolMenuDelegate::CreateLambda(&FillLinkAnimatorSubmenu, Animator, InMenuData)
+				, bOpenOnClick
+				, MenuIcon
+				, bCloseMenuAfterSelection
+			);
+		}
+	}
+}
+
+void UE::PropertyAnimatorCoreEditor::Menu::FillLinkAnimatorSection(UToolMenu* InMenu, TSharedRef<FPropertyAnimatorCoreEditorMenuData> InMenuData)
+{
+	if (!InMenu || InMenuData->GetContext().IsEmpty())
+	{
+		return;
+	}
+
+	const UPropertyAnimatorCoreSubsystem* Subsystem = UPropertyAnimatorCoreSubsystem::Get();
+	if (!Subsystem)
+	{
+		return;
+	}
+
+	const TSet<UPropertyAnimatorCoreBase*> Animators = InMenuData->GetContext().GetAnimators();
+
+	if (Animators.Num() != 1)
+	{
+		return;
+	}
+
+	UPropertyAnimatorCoreBase* Animator = Animators.Array()[0];
+
+	if (!Animator)
+	{
+		return;
+	}
+
+	FillLinkAnimatorSubmenu(InMenu, Animator, InMenuData);
+}
+
+void UE::PropertyAnimatorCoreEditor::Menu::FillDeleteAnimatorSection(UToolMenu* InMenu, TSharedRef<FPropertyAnimatorCoreEditorMenuData> InMenuData)
+{
+	if (!InMenu || !InMenuData->GetContext().ContainsAnyComponent())
+	{
+		return;
+	}
+
+	const UPropertyAnimatorCoreSubsystem* Subsystem = UPropertyAnimatorCoreSubsystem::Get();
+	if (!Subsystem)
+	{
+		return;
+	}
+
+	FToolMenuSection& ActorAnimatorsSection = InMenu->FindOrAddSection(TEXT("ActorAnimators"), LOCTEXT("ActorAnimators.Label", "Actor Animators"));
+
+	ActorAnimatorsSection.AddMenuEntry(
+		TEXT("DeleteActorAnimator")
+		, LOCTEXT("DeleteActorAnimator.Label", "Delete actor animators")
+		, LOCTEXT("DeleteActorAnimator.Tooltip", "Delete selected actor animators")
+		, FSlateIcon()
+		, FUIAction(
+		  FExecuteAction::CreateLambda(&ExecuteDeleteActorAnimatorAction, InMenuData)
+		)
+	);
+
+	const TSet<UPropertyAnimatorCoreBase*>& Animators = InMenuData->GetContext().GetAnimators();
+
+	if (Animators.IsEmpty())
+	{
+		return;
+	}
+
+	ActorAnimatorsSection.AddSeparator(TEXT("ActorAnimatorSeparator"));
+
+	for (UPropertyAnimatorCoreBase* Animator : Animators)
+	{
+		if (!IsValid(Animator))
+		{
+			continue;
+		}
+
+		ActorAnimatorsSection.AddMenuEntry(
+			FName(Animator->GetAnimatorDisplayName())
+			, FText::Format(LOCTEXT("DeleteActorAnimator.Label", "Delete {0}"), FText::FromString(Animator->GetAnimatorDisplayName()))
+			, LOCTEXT("DeleteActorAnimator.Tooltip", "Delete selected animator")
+			, FSlateIcon()
+			, FUIAction(
+			  FExecuteAction::CreateLambda(&ExecuteDeleteAnimatorAction, Animator, InMenuData)
+			)
+		);
+	}
+}
+
+void UE::PropertyAnimatorCoreEditor::Menu::FillEnableAnimatorSection(UToolMenu* InMenu, TSharedRef<FPropertyAnimatorCoreEditorMenuData> InMenuData)
+{
+	if (!InMenu
+		|| InMenuData->GetContext().IsEmpty()
+		|| !InMenuData->GetContext().ContainsAnyDisabledAnimator())
+	{
+		return;
+	}
+
+	const UPropertyAnimatorCoreSubsystem* Subsystem = UPropertyAnimatorCoreSubsystem::Get();
+	if (!Subsystem)
+	{
+		return;
+	}
+
+	FToolMenuSection& ActorAnimatorsSection = InMenu->FindOrAddSection(TEXT("ActorAnimators"), LOCTEXT("ActorAnimators.Label", "Actor Animators"));
+
+	constexpr bool bEnable = true;
+
+	ActorAnimatorsSection.AddMenuEntry(
+		TEXT("EnableActorAnimator")
+		, LOCTEXT("EnableActorAnimator.Label", "Enable actor animators")
+		, LOCTEXT("EnableActorAnimator.Tooltip", "Enable selected actor animators")
+		, FSlateIcon()
+		, FUIAction(
+			FExecuteAction::CreateLambda(&ExecuteEnableActorAnimatorAction, InMenuData, bEnable)
+		)
+	);
+
+	ActorAnimatorsSection.AddMenuEntry(
+		TEXT("EnableLevelAnimator")
+		, LOCTEXT("EnableLevelAnimator.Label", "Enable level animators")
+		, LOCTEXT("EnableLevelAnimator.Tooltip", "Enable current level animators")
+		, FSlateIcon()
+		, FUIAction(
+			FExecuteAction::CreateLambda(&ExecuteEnableLevelAnimatorAction, InMenuData, bEnable)
+		)
+	);
+
+	const TSet<UPropertyAnimatorCoreBase*> DisabledAnimators = InMenuData->GetContext().GetDisabledAnimators();
+
+	if (DisabledAnimators.IsEmpty())
+	{
+		return;
+	}
+
+	ActorAnimatorsSection.AddSeparator(TEXT("ActorAnimatorSeparator"));
+
+	for (UPropertyAnimatorCoreBase* Animator : DisabledAnimators)
+	{
+		if (!IsValid(Animator))
+		{
+			continue;
+		}
+
+		ActorAnimatorsSection.AddMenuEntry(
+			FName(Animator->GetAnimatorDisplayName())
+			, FText::Format(LOCTEXT("EnableAnimator.Label", "Enable {0}"), FText::FromString(Animator->GetAnimatorDisplayName()))
+			, LOCTEXT("EnableAnimator.Tooltip", "Enable selected animator")
+			, FSlateIconFinder::FindIconForClass(Animator->GetClass())
+			, FUIAction(
+				FExecuteAction::CreateLambda(&ExecuteEnableAnimatorAction, Animator, bEnable, InMenuData)
+			)
+		);
+	}
+}
+
+void UE::PropertyAnimatorCoreEditor::Menu::FillDisableAnimatorSection(UToolMenu* InMenu, TSharedRef<FPropertyAnimatorCoreEditorMenuData> InMenuData)
+{
+	if (!InMenu
+		|| InMenuData->GetContext().IsEmpty()
+		|| !InMenuData->GetContext().ContainsAnyEnabledAnimator())
+	{
+		return;
+	}
+
+	const UPropertyAnimatorCoreSubsystem* Subsystem = UPropertyAnimatorCoreSubsystem::Get();
+	if (!Subsystem)
+	{
+		return;
+	}
+
+	FToolMenuSection& ActorAnimatorsSection = InMenu->FindOrAddSection(TEXT("ActorAnimators"), LOCTEXT("ActorAnimators.Label", "Actor Animators"));
+
+	constexpr bool bEnable = false;
+
+	ActorAnimatorsSection.AddMenuEntry(
+		TEXT("DisableActorAnimator")
+		, LOCTEXT("DisableActorAnimator.Label", "Disable actor animators")
+		, LOCTEXT("DisableActorAnimator.Tooltip", "Disable selected actor animators")
+		, FSlateIcon()
+		, FUIAction(
+			FExecuteAction::CreateLambda(&ExecuteEnableActorAnimatorAction, InMenuData, bEnable)
+		)
+	);
+
+	ActorAnimatorsSection.AddMenuEntry(
+		TEXT("DisableLevelAnimator")
+		, LOCTEXT("DisableLevelAnimator.Label", "Disable level animators")
+		, LOCTEXT("DisableLevelAnimator.Tooltip", "Disable current level animators")
+		, FSlateIcon()
+		, FUIAction(
+			FExecuteAction::CreateLambda(&ExecuteEnableLevelAnimatorAction, InMenuData, bEnable)
+		)
+	);
+
+	const TSet<UPropertyAnimatorCoreBase*> EnabledAnimators = InMenuData->GetContext().GetEnabledAnimators();
+
+	if (EnabledAnimators.IsEmpty())
+	{
+		return;
+	}
+
+	ActorAnimatorsSection.AddSeparator(TEXT("ActorAnimatorSeparator"));
+
+	for (UPropertyAnimatorCoreBase* Animator : EnabledAnimators)
+	{
+		if (!IsValid(Animator))
+		{
+			continue;
+		}
+
+		ActorAnimatorsSection.AddMenuEntry(
+			FName(Animator->GetAnimatorDisplayName())
+			, FText::Format(LOCTEXT("DisableAnimator.Label", "Disable {0}"), FText::FromString(Animator->GetAnimatorDisplayName()))
+			, LOCTEXT("DisableAnimator.Tooltip", "Disable selected animator")
+			, FSlateIconFinder::FindIconForClass(Animator->GetClass())
+			, FUIAction(
+				FExecuteAction::CreateLambda(&ExecuteEnableAnimatorAction, Animator, bEnable, InMenuData)
+			)
+		);
+	}
+}
+
+void UE::PropertyAnimatorCoreEditor::Menu::ExecuteEditAnimatorAction(AActor* InActor)
+{
+	UPropertyAnimatorCoreEditorSubsystem* EditorSubsystem = UPropertyAnimatorCoreEditorSubsystem::Get();
+
+	if (!EditorSubsystem || !IsValid(InActor))
+	{
+		return;
+	}
+
+	FPropertyAnimatorCoreEditorEditPanelOptions& Options = EditorSubsystem->OpenPropertyControlWindow();
+	Options.SetContextActor(InActor);
+}
+
+void UE::PropertyAnimatorCoreEditor::Menu::ExecuteNewAnimatorPresetAction(const UPropertyAnimatorCoreBase* InAnimator, const TSet<AActor*>& InActors, UPropertyAnimatorCorePresetBase* InPreset, TSharedRef<FPropertyAnimatorCoreEditorMenuData> InMenuData)
+{
+	const UPropertyAnimatorCoreSubsystem* Subsystem = UPropertyAnimatorCoreSubsystem::Get();
+
+	if (!Subsystem
+		|| !IsValid(InAnimator)
+		|| !InAnimator->IsTemplate()
+		|| InMenuData->GetContext().IsEmpty())
+	{
+		return;
+	}
+
+	InMenuData->SetLastCreatedAnimators(Subsystem->CreateAnimators(InActors, InAnimator->GetClass(), InPreset, InMenuData->GetOptions().ShouldTransact()));
+}
+
+void UE::PropertyAnimatorCoreEditor::Menu::ExecuteNewAnimatorPropertyAction(const UPropertyAnimatorCoreBase* InAnimator, FPropertyAnimatorCoreData InProperty, TSharedRef<FPropertyAnimatorCoreEditorMenuData> InMenuData)
+{
+	UPropertyAnimatorCoreSubsystem* Subsystem = UPropertyAnimatorCoreSubsystem::Get();
+
+	if (!Subsystem
+		|| !IsValid(InAnimator)
+		|| !InAnimator->IsTemplate()
+		|| InMenuData->GetContext().IsEmpty())
+	{
+		return;
+	}
+
+	AActor* OwningActor = InProperty.GetOwningActor();
+
+	UPropertyAnimatorCoreBase* NewAnimator = Subsystem->CreateAnimator(OwningActor, InAnimator->GetClass(), nullptr, InMenuData->GetOptions().ShouldTransact());
+	Subsystem->LinkAnimatorProperty(NewAnimator, InProperty, true);
+	InMenuData->SetLastCreatedAnimator(NewAnimator);
+}
+
+void UE::PropertyAnimatorCoreEditor::Menu::ExecuteLinkLastCreatedAnimatorPropertyAction(const UPropertyAnimatorCoreBase* InAnimator, FPropertyAnimatorCoreData InProperty, TSharedRef<FPropertyAnimatorCoreEditorMenuData> InMenuData)
+{
+	UPropertyAnimatorCoreSubsystem* Subsystem = UPropertyAnimatorCoreSubsystem::Get();
+
+	if (!Subsystem
+		|| !IsValid(InAnimator)
+		|| !InAnimator->IsTemplate())
+	{
+		return;
+	}
+
+	for (UPropertyAnimatorCoreBase* LastCreatedAnimator : InMenuData->GetLastCreatedAnimators())
+	{
+		if (LastCreatedAnimator->GetClass() != InAnimator->GetClass())
+		{
+			continue;
+		}
+
+		if (LastCreatedAnimator->IsPropertyLinked(InProperty))
+		{
+			Subsystem->UnlinkAnimatorProperty(LastCreatedAnimator, InProperty, true);
+		}
+		else
+		{
+			Subsystem->LinkAnimatorProperty(LastCreatedAnimator, InProperty, true);
+		}
+	}
+}
+
+void UE::PropertyAnimatorCoreEditor::Menu::ExecuteApplyLastCreatedAnimatorPresetAction(const UPropertyAnimatorCoreBase* InAnimator, UPropertyAnimatorCorePresetBase* InPreset, TSharedRef<FPropertyAnimatorCoreEditorMenuData> InMenuData)
+{
+	UPropertyAnimatorCoreSubsystem* Subsystem = UPropertyAnimatorCoreSubsystem::Get();
+
+	if (!Subsystem
+		|| !IsValid(InAnimator)
+		|| !InAnimator->IsTemplate())
+	{
+		return;
+	}
+
+	for (UPropertyAnimatorCoreBase* LastCreatedAnimator : InMenuData->GetLastCreatedAnimators())
+	{
+		if (LastCreatedAnimator->GetClass() != InAnimator->GetClass())
+		{
+			continue;
+		}
+
+		if (InPreset->IsPresetApplied(LastCreatedAnimator))
+		{
+			Subsystem->UnapplyAnimatorPreset(LastCreatedAnimator, InPreset, true);
+		}
+		else
+		{
+			Subsystem->ApplyAnimatorPreset(LastCreatedAnimator, InPreset, true);
+		}
+	}
+}
+
+void UE::PropertyAnimatorCoreEditor::Menu::ExecuteLinkAnimatorPresetAction(UPropertyAnimatorCoreBase* InAnimator, UPropertyAnimatorCorePresetBase* InPreset, TSharedRef<FPropertyAnimatorCoreEditorMenuData> InMenuData)
+{
+	UPropertyAnimatorCoreSubsystem* Subsystem = UPropertyAnimatorCoreSubsystem::Get();
+
+	if (!Subsystem
+		|| !InPreset
+		|| !IsValid(InAnimator)
+		|| InAnimator->IsTemplate()
+		|| InMenuData->GetContext().IsEmpty())
+	{
+		return;
+	}
+
+	if (InPreset->IsPresetApplied(InAnimator))
+	{
+		Subsystem->UnapplyAnimatorPreset(InAnimator, InPreset, InMenuData->GetOptions().ShouldTransact());
+	}
+	else
+	{
+		Subsystem->ApplyAnimatorPreset(InAnimator, InPreset, InMenuData->GetOptions().ShouldTransact());
+	}
+}
+
+void UE::PropertyAnimatorCoreEditor::Menu::ExecuteLinkAnimatorPropertyAction(UPropertyAnimatorCoreBase* InAnimator, FPropertyAnimatorCoreData InProperty, TSharedRef<FPropertyAnimatorCoreEditorMenuData> InMenuData)
+{
+	UPropertyAnimatorCoreSubsystem* Subsystem = UPropertyAnimatorCoreSubsystem::Get();
+
+	if (!Subsystem
+		|| !IsValid(InAnimator)
+		|| InAnimator->IsTemplate()
+		|| !InProperty.IsResolved()
+		|| !InAnimator->IsPropertySupported(InProperty)
+		|| InMenuData->GetContext().IsEmpty())
+	{
+		return;
+	}
+
+	if (InAnimator->IsPropertyLinked(InProperty))
+	{
+		Subsystem->UnlinkAnimatorProperty(InAnimator, InProperty, InMenuData->GetOptions().ShouldTransact());
+	}
+	else
+	{
+		Subsystem->LinkAnimatorProperty(InAnimator, InProperty, InMenuData->GetOptions().ShouldTransact());
+	}
+}
+
+void UE::PropertyAnimatorCoreEditor::Menu::ExecuteEnableActorAnimatorAction(TSharedRef<FPropertyAnimatorCoreEditorMenuData> InMenuData, bool bInEnable)
+{
+	UPropertyAnimatorCoreSubsystem* Subsystem = UPropertyAnimatorCoreSubsystem::Get();
+
+	if (InMenuData->GetContext().IsEmpty() || !Subsystem)
+	{
+		return;
+	}
+
+	Subsystem->SetActorAnimatorsEnabled(InMenuData->GetContext().GetActors(), bInEnable, InMenuData->GetOptions().ShouldTransact());
+}
+
+void UE::PropertyAnimatorCoreEditor::Menu::ExecuteEnableLevelAnimatorAction(TSharedRef<FPropertyAnimatorCoreEditorMenuData> InMenuData, bool bInEnable)
+{
+	UPropertyAnimatorCoreSubsystem* Subsystem = UPropertyAnimatorCoreSubsystem::Get();
+	const UWorld* World = InMenuData->GetContext().GetWorld();
+
+	if (!IsValid(World) || !Subsystem)
+	{
+		return;
+	}
+
+	Subsystem->SetLevelAnimatorsEnabled(World, bInEnable, InMenuData->GetOptions().ShouldTransact());
+}
+
+void UE::PropertyAnimatorCoreEditor::Menu::ExecuteEnableAnimatorAction(UPropertyAnimatorCoreBase* InAnimator, bool bInEnable, TSharedRef<FPropertyAnimatorCoreEditorMenuData> InMenuData)
+{
+	UPropertyAnimatorCoreSubsystem* Subsystem = UPropertyAnimatorCoreSubsystem::Get();
+
+	if (!IsValid(InAnimator) || !Subsystem)
+	{
+		return;
+	}
+
+	Subsystem->SetAnimatorsEnabled({InAnimator}, bInEnable, InMenuData->GetOptions().ShouldTransact());
+}
+
+void UE::PropertyAnimatorCoreEditor::Menu::ExecuteDeleteActorAnimatorAction(TSharedRef<FPropertyAnimatorCoreEditorMenuData> InMenuData)
+{
+	const UPropertyAnimatorCoreSubsystem* Subsystem = UPropertyAnimatorCoreSubsystem::Get();
+
+	if (!Subsystem || InMenuData->GetContext().IsEmpty())
+	{
+		return;
+	}
+
+	TSet<UPropertyAnimatorCoreBase*> Animators;
+
+	for (const UPropertyAnimatorCoreComponent* Component : InMenuData->GetContext().GetComponents())
+	{
+		if (!IsValid(Component))
+		{
+			continue;
+		}
+
+		for (const TObjectPtr<UPropertyAnimatorCoreBase>& Animator : Component->GetAnimators())
+		{
+			if (Animator)
+			{
+				Animators.Add(Animator);
+			}
+		}
+	}
+
+	Subsystem->RemoveAnimators(Animators, InMenuData->GetOptions().ShouldTransact());
+}
+
+void UE::PropertyAnimatorCoreEditor::Menu::ExecuteDeleteAnimatorAction(UPropertyAnimatorCoreBase* InAnimator, TSharedRef<FPropertyAnimatorCoreEditorMenuData> InMenuData)
+{
+	const UPropertyAnimatorCoreSubsystem* Subsystem = UPropertyAnimatorCoreSubsystem::Get();
+
+	if (!IsValid(InAnimator) || !Subsystem)
+	{
+		return;
+	}
+
+	Subsystem->RemoveAnimator(InAnimator, InMenuData->GetOptions().ShouldTransact());
+}
+
+bool UE::PropertyAnimatorCoreEditor::Menu::IsAnimatorPresetLinked(UPropertyAnimatorCoreBase* InAnimator, UPropertyAnimatorCorePresetBase* InPreset)
+{
+	if (!IsValid(InAnimator) || InAnimator->IsTemplate() || !InPreset)
+	{
+		return false;
+	}
+
+	return InPreset->IsPresetApplied(InAnimator);
+}
+
+void UE::PropertyAnimatorCoreEditor::Menu::FillNewAnimatorSubmenu(UToolMenu* InMenu, UPropertyAnimatorCoreBase* InAnimator, TSharedRef<FPropertyAnimatorCoreEditorMenuData> InMenuData)
+{
+	const UPropertyAnimatorCoreSubsystem* Subsystem = UPropertyAnimatorCoreSubsystem::Get();
+
+	if (!InMenu
+		|| !Subsystem
+		|| !InAnimator
+		|| !InAnimator->IsTemplate())
+	{
+		return;
+	}
+
+	if (InMenuData->GetContext().ContainsAnyActor())
+	{
+		const TSet<AActor*>& ContextActors = InMenuData->GetContext().GetActors();
+
+		FToolMenuSection& PresetSection = InMenu->FindOrAddSection(TEXT("Presets"), LOCTEXT("NewAnimatorPresetsSection.Label", "Presets"));
+
+		UPropertyAnimatorCorePresetBase* EmptyPreset = nullptr;
+
+		PresetSection.AddMenuEntry(
+			TEXT("EmptyPreset")
+			, LOCTEXT("NewAnimatorEmptyPresetSection.Label", "Empty")
+			, LOCTEXT("NewAnimatorEmptyPresetSection.Tooltip", "Create an empty animator")
+			, FSlateIcon()
+			, FUIAction(
+				FExecuteAction::CreateLambda(&ExecuteNewAnimatorPresetAction, InAnimator, ContextActors, EmptyPreset, InMenuData)
+				, FCanExecuteAction()
+				, FIsActionChecked()
+				, FIsActionButtonVisible::CreateLambda(&IsLastAnimatorCreatedActionHidden, InAnimator, InMenuData)
+			)
+		);
+
+		TSet<UPropertyAnimatorCorePresetBase*> SupportedPresets = Subsystem->GetAvailablePresets();
+
+		for (const AActor* Actor : ContextActors)
+		{
+			SupportedPresets = SupportedPresets.Intersect(Subsystem->GetSupportedPresets(Actor, InAnimator));
+		}
+
+		for (UPropertyAnimatorCorePresetBase* SupportedPreset : SupportedPresets)
+		{
+			if (!SupportedPreset)
+			{
+				continue;
+			}
+
+			const FString MenuName = SupportedPreset->GetPresetName().ToString();
+			const FText MenuLabel = FText::FromString(SupportedPreset->GetPresetDisplayName());
+
+			// Create action (creates an animator and links the property)
+			PresetSection.AddMenuEntry(
+				FName(TEXT("Create") + MenuName)
+				, MenuLabel
+				, LOCTEXT("NewAnimatorPresetSection.Tooltip", "Create this animator using this preset")
+				, FSlateIcon()
+				, FUIAction(
+					FExecuteAction::CreateLambda(&ExecuteNewAnimatorPresetAction, InAnimator, ContextActors, SupportedPreset, InMenuData)
+					, FCanExecuteAction()
+					, FIsActionChecked()
+					, FIsActionButtonVisible::CreateLambda(&IsLastAnimatorCreatedActionHidden, InAnimator, InMenuData)
+				)
+			);
+
+			// Link action (links the property to last created animator)
+			PresetSection.AddMenuEntry(
+				FName(TEXT("Apply") + MenuName)
+				, MenuLabel
+				, LOCTEXT("ApplyLastCreatedAnimatorPresetSection.Tooltip", "Apply this preset to the last created animator")
+				, FSlateIcon()
+				, FUIAction(
+					FExecuteAction::CreateLambda(&ExecuteApplyLastCreatedAnimatorPresetAction, InAnimator, SupportedPreset, InMenuData)
+					, FCanExecuteAction()
+					, FIsActionChecked::CreateLambda(&IsLastAnimatorCreatedPresetLinked, InAnimator, SupportedPreset, InMenuData)
+					, FIsActionButtonVisible::CreateLambda(&IsLastAnimatorCreatedActionVisible, InAnimator, InMenuData)
+				)
+				, EUserInterfaceActionType::ToggleButton
+			);
+		}
+	}
+
+	if (InMenuData->GetContext().ContainsAnyProperty())
+	{
+		FToolMenuSection& PropertySection = InMenu->FindOrAddSection(TEXT("Properties"), LOCTEXT("NewAnimatorPropertiesSection.Label", "Properties"));
+
+		TSet<FPropertyAnimatorCoreData> SupportedProperties;
+		for (const FPropertyAnimatorCoreData& Property : InMenuData->GetContext().GetProperties())
+		{
+			InAnimator->GetPropertiesSupported(Property, SupportedProperties, true);
+		}
+
+		for (FPropertyAnimatorCoreData& SupportedProperty : SupportedProperties)
+		{
+			const FString MenuName = SupportedProperty.GetPropertyDisplayName().ToString();
+			const FText MenuLabel = FText::FromString(MenuName);
+
+			// Create action (creates an animator and links the property)
+			PropertySection.AddMenuEntry(
+				FName(TEXT("Create") + MenuName)
+				, MenuLabel
+				, LOCTEXT("NewAnimatorPropertySection.Tooltip", "Create this animator using this property")
+				, FSlateIcon()
+				, FUIAction(
+					FExecuteAction::CreateLambda(&ExecuteNewAnimatorPropertyAction, InAnimator, SupportedProperty, InMenuData)
+					, FCanExecuteAction()
+					, FIsActionChecked()
+					, FIsActionButtonVisible::CreateLambda(&IsLastAnimatorCreatedActionHidden, InAnimator, InMenuData)
+				)
+			);
+
+			// Link action (links the property to last created animator)
+			PropertySection.AddMenuEntry(
+				FName(TEXT("Link") + MenuName)
+				, MenuLabel
+				, LOCTEXT("LinkLastCreatedAnimatorPropertySection.Tooltip", "Link this property to the last created animator")
+				, FSlateIcon()
+				, FUIAction(
+					FExecuteAction::CreateLambda(&ExecuteLinkLastCreatedAnimatorPropertyAction, InAnimator, SupportedProperty, InMenuData)
+					, FCanExecuteAction()
+					, FIsActionChecked::CreateLambda(&IsLastAnimatorCreatedPropertyLinked, InAnimator, SupportedProperty, InMenuData)
+					, FIsActionButtonVisible::CreateLambda(&IsLastAnimatorCreatedActionVisible, InAnimator, InMenuData)
+				)
+				, EUserInterfaceActionType::ToggleButton
+			);
+		}
+	}
+}
+
+void UE::PropertyAnimatorCoreEditor::Menu::FillLinkAnimatorSubmenu(UToolMenu* InMenu, UPropertyAnimatorCoreBase* InAnimator, TSharedRef<FPropertyAnimatorCoreEditorMenuData> InMenuData)
+{
+	if (!InMenu
+		|| !IsValid(InAnimator)
+		|| InAnimator->IsTemplate()
+		|| InMenuData->GetContext().IsEmpty())
+	{
+		return;
+	}
+
+	const UPropertyAnimatorCoreSubsystem* Subsystem = UPropertyAnimatorCoreSubsystem::Get();
+	if (!Subsystem)
+	{
+		return;
+	}
+
+	FToolMenuSection& PresetSection = InMenu->FindOrAddSection(TEXT("Presets"), LOCTEXT("LinkAnimatorPresetsSection.Label", "Presets"));
+
+	TSet<UPropertyAnimatorCorePresetBase*> SupportedPresets = Subsystem->GetAvailablePresets();
+
+	for (const AActor* Actor : InMenuData->GetContext().GetActors())
+	{
+		SupportedPresets = SupportedPresets.Intersect(Subsystem->GetSupportedPresets(Actor, InAnimator));
+	}
+
+	for (UPropertyAnimatorCorePresetBase* SupportedPreset : SupportedPresets)
+	{
+		if (!SupportedPreset)
+		{
+			continue;
+		}
+
+		const FName MenuName = SupportedPreset->GetPresetName();
+		const FText MenuLabel = FText::FromString(SupportedPreset->GetPresetDisplayName());
+
+		PresetSection.AddMenuEntry(
+			MenuName
+			, MenuLabel
+			, LOCTEXT("LinkAnimatorPresetSection.Tooltip", "Link or unlink a preset from this animator")
+			, FSlateIcon()
+			, FUIAction(
+				FExecuteAction::CreateLambda(&ExecuteLinkAnimatorPresetAction, InAnimator, SupportedPreset, InMenuData)
+				, FCanExecuteAction()
+				, FIsActionChecked::CreateStatic(&IsAnimatorPresetLinked, InAnimator, SupportedPreset)
+			)
+			, EUserInterfaceActionType::ToggleButton
+		);
+	}
+
+	FToolMenuSection& PropertySection = InMenu->FindOrAddSection(TEXT("Properties"), LOCTEXT("LinkAnimatorPropertiesSection.Label", "Properties"));
+
+	TSet<FPropertyAnimatorCoreData> SupportedProperties;
+	for (const FPropertyAnimatorCoreData& Property : InMenuData->GetContext().GetProperties())
+	{
+		InAnimator->GetPropertiesSupported(Property, SupportedProperties, true);
+	}
+
+	for (const FPropertyAnimatorCoreData& SupportedProperty : SupportedProperties)
+	{
+		const FName MenuName = SupportedProperty.GetPropertyDisplayName();
+		const FText MenuLabel = FText::FromName(MenuName);
+
+		PropertySection.AddMenuEntry(
+			MenuName
+			, MenuLabel
+			, LOCTEXT("LinkAnimatorPropertySection.Tooltip", "Link or unlink this property from the animator")
+			, FSlateIcon()
+			, FUIAction(
+				FExecuteAction::CreateLambda(&ExecuteLinkAnimatorPropertyAction, InAnimator, SupportedProperty, InMenuData)
+				, FCanExecuteAction()
+				, FIsActionChecked::CreateLambda(&IsAnimatorPropertyLinked, InAnimator, SupportedProperty)
+			)
+			, EUserInterfaceActionType::ToggleButton
+		);
+	}
+}
+
+bool UE::PropertyAnimatorCoreEditor::Menu::IsAnimatorPropertyLinked(const UPropertyAnimatorCoreBase* InAnimator, const FPropertyAnimatorCoreData& InProperty)
+{
+	if (!IsValid(InAnimator) || InAnimator->IsTemplate())
+	{
+		return false;
+	}
+
+	return InAnimator->IsPropertyLinked(InProperty);
+}
+
+bool UE::PropertyAnimatorCoreEditor::Menu::IsLastAnimatorCreatedPropertyLinked(const UPropertyAnimatorCoreBase* InAnimator, const FPropertyAnimatorCoreData& InProperty, TSharedRef<FPropertyAnimatorCoreEditorMenuData> InMenuData)
+{
+	if (!IsValid(InAnimator)
+		|| !InAnimator->IsTemplate()
+		|| !InMenuData->ContainsAnyLastCreatedAnimator())
+	{
+		return false;
+	}
+
+	for (const UPropertyAnimatorCoreBase* LastCreatedAnimator : InMenuData->GetLastCreatedAnimators())
+	{
+		if (LastCreatedAnimator->GetClass() != InAnimator->GetClass()
+			|| !LastCreatedAnimator->IsPropertyLinked(InProperty))
+		{
+			return false;
+		}
+	}
+
+	return true;
+}
+
+bool UE::PropertyAnimatorCoreEditor::Menu::IsLastAnimatorCreatedPresetLinked(const UPropertyAnimatorCoreBase* InAnimator, UPropertyAnimatorCorePresetBase* InPreset, TSharedRef<FPropertyAnimatorCoreEditorMenuData> InMenuData)
+{
+	if (!IsValid(InAnimator)
+		|| !InAnimator->IsTemplate()
+		|| !InPreset
+		|| !InMenuData->ContainsAnyLastCreatedAnimator())
+	{
+		return false;
+	}
+
+	for (const UPropertyAnimatorCoreBase* LastCreatedAnimator : InMenuData->GetLastCreatedAnimators())
+	{
+		if (LastCreatedAnimator->GetClass() != InAnimator->GetClass()
+			|| !InPreset->IsPresetApplied(LastCreatedAnimator))
+		{
+			return false;
+		}
+	}
+
+	return true;
+}
+
+bool UE::PropertyAnimatorCoreEditor::Menu::IsLastAnimatorCreatedActionVisible(const UPropertyAnimatorCoreBase* InAnimator, TSharedRef<FPropertyAnimatorCoreEditorMenuData> InMenuData)
+{
+	if (!IsValid(InAnimator)
+		|| !InAnimator->IsTemplate()
+		|| !InMenuData->ContainsAnyLastCreatedAnimator())
+	{
+		return false;
+	}
+
+	for (const UPropertyAnimatorCoreBase* LastCreatedAnimator : InMenuData->GetLastCreatedAnimators())
+	{
+		if (LastCreatedAnimator->GetClass() != InAnimator->GetClass())
+		{
+			return false;
+		}
+	}
+
+	return true;
+}
+
+bool UE::PropertyAnimatorCoreEditor::Menu::IsLastAnimatorCreatedActionHidden(const UPropertyAnimatorCoreBase* InAnimator, TSharedRef<FPropertyAnimatorCoreEditorMenuData> InMenuData)
+{
+	return !IsLastAnimatorCreatedActionVisible(InAnimator, InMenuData);
+}
+
+#undef LOCTEXT_NAMESPACE
diff --git a/Engine/Plugins/Experimental/PropertyAnimatorCore/Source/PropertyAnimatorCoreEditor/Private/Menus/PropertyAnimatorCoreEditorMenu.h b/Engine/Plugins/Experimental/PropertyAnimatorCore/Source/PropertyAnimatorCoreEditor/Private/Menus/PropertyAnimatorCoreEditorMenu.h
new file mode 100644
--- /dev/null
+++ b/Engine/Plugins/Experimental/PropertyAnimatorCore/Source/PropertyAnimatorCoreEditor/Private/Menus/PropertyAnimatorCoreEditorMenu.h
@@ -0,0 +1,81 @@
+// Copyright Epic Games, Inc. All Rights Reserved.
+
+#pragma once
+
+#include "Misc/Optional.h"
+#include "Menus/PropertyAnimatorCoreEditorMenuDefs.h"
+#include "Presets/PropertyAnimatorCorePresetBase.h"
+#include "Properties/PropertyAnimatorCoreData.h"
+#include "Templates/SharedPointer.h"
+#include "ToolMenus.h"
+
+struct FPropertyAnimatorCoreData;
+class IPropertyHandle;
+class UPropertyAnimatorCoreBase;
+class UPropertyAnimatorCoreEditorMenuContext;
+class UPropertyAnimatorCorePresetBase;
+
+namespace UE::PropertyAnimatorCoreEditor::Menu
+{
+	/** Sections */
+
+	void FillEditAnimatorSection(UToolMenu* InMenu, TSharedRef<FPropertyAnimatorCoreEditorMenuData> InMenuData);
+
+	void FillNewAnimatorSection(UToolMenu* InMenu, TSharedRef<FPropertyAnimatorCoreEditorMenuData> InMenuData);
+
+	void FillExistingAnimatorSection(UToolMenu* InMenu, TSharedRef<FPropertyAnimatorCoreEditorMenuData> InMenuData);
+
+	void FillLinkAnimatorSection(UToolMenu* InMenu, TSharedRef<FPropertyAnimatorCoreEditorMenuData> InMenuData);
+
+	void FillEnableAnimatorSection(UToolMenu* InMenu, TSharedRef<FPropertyAnimatorCoreEditorMenuData> InMenuData);
+
+	void FillDisableAnimatorSection(UToolMenu* InMenu, TSharedRef<FPropertyAnimatorCoreEditorMenuData> InMenuData);
+
+	void FillDeleteAnimatorSection(UToolMenu* InMenu, TSharedRef<FPropertyAnimatorCoreEditorMenuData> InMenuData);
+
+	/** Submenus */
+
+	void FillNewAnimatorSubmenu(UToolMenu* InMenu, UPropertyAnimatorCoreBase* InAnimator, TSharedRef<FPropertyAnimatorCoreEditorMenuData> InMenuData);
+
+	void FillLinkAnimatorSubmenu(UToolMenu* InMenu, UPropertyAnimatorCoreBase* InAnimator, TSharedRef<FPropertyAnimatorCoreEditorMenuData> InMenuData);
+
+	/** Execute */
+
+	void ExecuteEditAnimatorAction(AActor* InActor);
+
+	void ExecuteNewAnimatorPresetAction(const UPropertyAnimatorCoreBase* InAnimator, const TSet<AActor*>& InActors, UPropertyAnimatorCorePresetBase* InPreset, TSharedRef<FPropertyAnimatorCoreEditorMenuData> InMenuData);
+
+	void ExecuteNewAnimatorPropertyAction(const UPropertyAnimatorCoreBase* InAnimator, FPropertyAnimatorCoreData InProperty, TSharedRef<FPropertyAnimatorCoreEditorMenuData> InMenuData);
+
+	void ExecuteLinkLastCreatedAnimatorPropertyAction(const UPropertyAnimatorCoreBase* InAnimator, FPropertyAnimatorCoreData InProperty, TSharedRef<FPropertyAnimatorCoreEditorMenuData> InMenuData);
+
+	void ExecuteApplyLastCreatedAnimatorPresetAction(const UPropertyAnimatorCoreBase* InAnimator, UPropertyAnimatorCorePresetBase* InPreset, TSharedRef<FPropertyAnimatorCoreEditorMenuData> InMenuData);
+
+	void ExecuteLinkAnimatorPresetAction(UPropertyAnimatorCoreBase* InAnimator, UPropertyAnimatorCorePresetBase* InPreset, TSharedRef<FPropertyAnimatorCoreEditorMenuData> InMenuData);
+
+	void ExecuteLinkAnimatorPropertyAction(UPropertyAnimatorCoreBase* InAnimator, FPropertyAnimatorCoreData InProperty, TSharedRef<FPropertyAnimatorCoreEditorMenuData> InMenuData);
+
+	void ExecuteEnableActorAnimatorAction(TSharedRef<FPropertyAnimatorCoreEditorMenuData> InMenuData, bool bInEnable);
+
+	void ExecuteEnableLevelAnimatorAction(TSharedRef<FPropertyAnimatorCoreEditorMenuData> InMenuData, bool bInEnable);
+
+	void ExecuteEnableAnimatorAction(UPropertyAnimatorCoreBase* InAnimator, bool bInEnable, TSharedRef<FPropertyAnimatorCoreEditorMenuData> InMenuData);
+
+	void ExecuteDeleteActorAnimatorAction(TSharedRef<FPropertyAnimatorCoreEditorMenuData> InMenuData);
+
+	void ExecuteDeleteAnimatorAction(UPropertyAnimatorCoreBase* InAnimator, TSharedRef<FPropertyAnimatorCoreEditorMenuData> InMenuData);
+
+	/** Check */
+
+	bool IsAnimatorPresetLinked(UPropertyAnimatorCoreBase* InAnimator, UPropertyAnimatorCorePresetBase* InPreset);
+
+	bool IsAnimatorPropertyLinked(const UPropertyAnimatorCoreBase* InAnimator, const FPropertyAnimatorCoreData& InProperty);
+
+	bool IsLastAnimatorCreatedPropertyLinked(const UPropertyAnimatorCoreBase* InAnimator, const FPropertyAnimatorCoreData& InProperty, TSharedRef<FPropertyAnimatorCoreEditorMenuData> InMenuData);
+
+	bool IsLastAnimatorCreatedPresetLinked(const UPropertyAnimatorCoreBase* InAnimator, UPropertyAnimatorCorePresetBase* InPreset, TSharedRef<FPropertyAnimatorCoreEditorMenuData> InMenuData);
+
+	bool IsLastAnimatorCreatedActionVisible(const UPropertyAnimatorCoreBase* InAnimator, TSharedRef<FPropertyAnimatorCoreEditorMenuData> InMenuData);
+
+	bool IsLastAnimatorCreatedActionHidden(const UPropertyAnimatorCoreBase* InAnimator, TSharedRef<FPropertyAnimatorCoreEditorMenuData> InMenuData);
+}
diff --git a/Engine/Plugins/Experimental/PropertyAnimatorCore/Source/PropertyAnimatorCoreEditor/Private/Menus/PropertyAnimatorCoreEditorMenuContext.h b/Engine/Plugins/Experimental/PropertyAnimatorCore/Source/PropertyAnimatorCoreEditor/Private/Menus/PropertyAnimatorCoreEditorMenuContext.h
new file mode 100644
--- /dev/null
+++ b/Engine/Plugins/Experimental/PropertyAnimatorCore/Source/PropertyAnimatorCoreEditor/Private/Menus/PropertyAnimatorCoreEditorMenuContext.h
@@ -0,0 +1,30 @@
+// Copyright Epic Games, Inc. All Rights Reserved.
+
+#pragma once
+
+#include "Properties/PropertyAnimatorCoreData.h"
+#include "Templates/SharedPointer.h"
+#include "UObject/Object.h"
+#include "PropertyAnimatorCoreEditorMenuContext.generated.h"
+
+/** Context passed in UToolMenu when generating entries with selected items */
+UCLASS()
+class UPropertyAnimatorCoreEditorMenuContext : public UObject
+{
+	GENERATED_BODY()
+
+public:
+	const FPropertyAnimatorCoreData& GetPropertyData() const
+	{
+		return PropertyData;
+	}
+
+	void SetPropertyData(const FPropertyAnimatorCoreData& InPropertyData)
+	{
+		PropertyData = InPropertyData;
+	}
+
+protected:
+	/** The item this menu should apply to */
+	FPropertyAnimatorCoreData PropertyData;
+};
\ No newline at end of file
diff --git a/Engine/Plugins/Experimental/PropertyAnimatorCore/Source/PropertyAnimatorCoreEditor/Private/PropertyAnimatorCoreEditorModule.cpp b/Engine/Plugins/Experimental/PropertyAnimatorCore/Source/PropertyAnimatorCoreEditor/Private/PropertyAnimatorCoreEditorModule.cpp
new file mode 100644
--- /dev/null
+++ b/Engine/Plugins/Experimental/PropertyAnimatorCore/Source/PropertyAnimatorCoreEditor/Private/PropertyAnimatorCoreEditorModule.cpp
@@ -0,0 +1,29 @@
+// Copyright Epic Games, Inc. All Rights Reserved.
+
+#include "PropertyAnimatorCoreEditorModule.h"
+
+#include "Animators/PropertyAnimatorCoreBase.h"
+#include "Customizations/PropertyAnimatorCoreEditorDetailCustomization.h"
+#include "Customizations/PropertyAnimatorCoreEditorContextTypeCustomization.h"
+#include "Modules/ModuleManager.h"
+#include "Properties/PropertyAnimatorCoreContext.h"
+#include "PropertyEditorModule.h"
+
+void FPropertyAnimatorCoreEditorModule::StartupModule()
+{
+	FPropertyEditorModule& PropertyModule = FModuleManager::GetModuleChecked<FPropertyEditorModule>("PropertyEditor");
+	PropertyModule.RegisterCustomPropertyTypeLayout(UPropertyAnimatorCoreContext::StaticClass()->GetFName(), FOnGetPropertyTypeCustomizationInstance::CreateStatic(&FPropertyAnimatorCoreEditorContextTypeCustomization::MakeInstance));
+	PropertyModule.RegisterCustomClassLayout(UPropertyAnimatorCoreBase::StaticClass()->GetFName(), FOnGetDetailCustomizationInstance::CreateStatic(&FPropertyAnimatorCoreEditorDetailCustomization::MakeInstance));
+}
+
+void FPropertyAnimatorCoreEditorModule::ShutdownModule()
+{
+	if (FModuleManager::Get().IsModuleLoaded("PropertyEditor"))
+	{
+		FPropertyEditorModule& PropertyModule = FModuleManager::GetModuleChecked<FPropertyEditorModule>("PropertyEditor");
+		PropertyModule.UnregisterCustomPropertyTypeLayout(UPropertyAnimatorCoreContext::StaticClass()->GetFName());
+		PropertyModule.UnregisterCustomClassLayout(UPropertyAnimatorCoreBase::StaticClass()->GetFName());
+	}
+}
+
+IMPLEMENT_MODULE(FPropertyAnimatorCoreEditorModule, PropertyAnimatorCoreEditor)
diff --git a/Engine/Plugins/Experimental/PropertyAnimatorCore/Source/PropertyAnimatorCoreEditor/Private/Styles/PropertyAnimatorCoreEditorStyle.cpp b/Engine/Plugins/Experimental/PropertyAnimatorCore/Source/PropertyAnimatorCoreEditor/Private/Styles/PropertyAnimatorCoreEditorStyle.cpp
new file mode 100644
--- /dev/null
+++ b/Engine/Plugins/Experimental/PropertyAnimatorCore/Source/PropertyAnimatorCoreEditor/Private/Styles/PropertyAnimatorCoreEditorStyle.cpp
@@ -0,0 +1,37 @@
+// Copyright Epic Games, Inc. All Rights Reserved.
+
+#include "Styles/PropertyAnimatorCoreEditorStyle.h"
+
+#include "Interfaces/IPluginManager.h"
+#include "Misc/Paths.h"
+#include "Styling/SlateStyleMacros.h"
+#include "Styling/StyleColors.h"
+#include "Styling/SlateTypes.h"
+#include "Styling/SlateStyleRegistry.h"
+
+FPropertyAnimatorCoreEditorStyle::FPropertyAnimatorCoreEditorStyle()
+	: FSlateStyleSet(UE_MODULE_NAME)
+{
+	const FVector2f Icon16x16(16.f, 16.f);
+
+	const TSharedPtr<IPlugin> Plugin = IPluginManager::Get().FindPlugin(UE_PLUGIN_NAME);
+
+	check(Plugin.IsValid());
+
+	ContentRootDir = FPaths::Combine(Plugin->GetBaseDir(), TEXT("Resources"));
+
+	Set("ClassIcon.PropertyAnimatorCoreComponent",           new IMAGE_BRUSH_SVG("EditorIcons/PropertyAnimatorCoreDefault", Icon16x16));
+	Set("ClassIcon.PropertyAnimatorCoreBase",                new IMAGE_BRUSH_SVG("EditorIcons/PropertyAnimatorCoreDefault", Icon16x16));
+
+	Set("PropertyControlIcon.Default",                       new IMAGE_BRUSH_SVG("EditorIcons/PropertyAnimatorCoreDefault", Icon16x16));
+	Set("PropertyControlIcon.Linked",                        new IMAGE_BRUSH_SVG("EditorIcons/PropertyAnimatorCoreLinked",  Icon16x16));
+	Set("PropertyControlIcon.Add",                           new IMAGE_BRUSH("EditorIcons/PropertyAnimatorCoreAdd",         Icon16x16));
+	Set("PropertyControlIcon.Link",                          new IMAGE_BRUSH("EditorIcons/PropertyAnimatorCoreLink",        Icon16x16));
+
+	FSlateStyleRegistry::RegisterSlateStyle(*this);
+}
+
+FPropertyAnimatorCoreEditorStyle::~FPropertyAnimatorCoreEditorStyle()
+{
+	FSlateStyleRegistry::UnRegisterSlateStyle(*this);
+}
\ No newline at end of file
diff --git a/Engine/Plugins/Experimental/PropertyAnimatorCore/Source/PropertyAnimatorCoreEditor/Private/Widgets/DragDropOps/PropertyAnimatorCoreEditorViewDragDropOp.h b/Engine/Plugins/Experimental/PropertyAnimatorCore/Source/PropertyAnimatorCoreEditor/Private/Widgets/DragDropOps/PropertyAnimatorCoreEditorViewDragDropOp.h
new file mode 100644
--- /dev/null
+++ b/Engine/Plugins/Experimental/PropertyAnimatorCore/Source/PropertyAnimatorCoreEditor/Private/Widgets/DragDropOps/PropertyAnimatorCoreEditorViewDragDropOp.h
@@ -0,0 +1,28 @@
+// Copyright Epic Games, Inc. All Rights Reserved.
+
+#pragma once
+
+#include "DragAndDrop/DecoratedDragDropOp.h"
+#include "Widgets/Views/PropertyAnimatorCoreEditorViewItem.h"
+
+/** A custom drag operation for properties/controllers */
+class FPropertyAnimatorCoreEditorViewDragDropOp : public FDecoratedDragDropOp
+{
+public:
+	DRAG_DROP_OPERATOR_TYPE(FPropertyAnimatorCoreEditorViewDragDropOp, FDecoratedDragDropOp)
+
+	static TSharedRef<FPropertyAnimatorCoreEditorViewDragDropOp> New(const TSet<FPropertiesViewControllerItem>& InItems)
+	{
+		TSharedRef<FPropertyAnimatorCoreEditorViewDragDropOp> DragDropOp = MakeShared<FPropertyAnimatorCoreEditorViewDragDropOp>();
+		DragDropOp->DraggedItems = InItems;
+		return DragDropOp;
+	}
+
+	const TSet<FPropertiesViewControllerItem>& GetDraggedItems() const
+	{
+		return DraggedItems;
+	}
+
+protected:
+	TSet<FPropertiesViewControllerItem> DraggedItems;
+};
\ No newline at end of file
diff --git a/Engine/Plugins/Experimental/PropertyAnimatorCore/Source/PropertyAnimatorCoreEditor/Private/Widgets/PropertyAnimatorCoreEditorEditPanelOptions.cpp b/Engine/Plugins/Experimental/PropertyAnimatorCore/Source/PropertyAnimatorCoreEditor/Private/Widgets/PropertyAnimatorCoreEditorEditPanelOptions.cpp
new file mode 100644
--- /dev/null
+++ b/Engine/Plugins/Experimental/PropertyAnimatorCore/Source/PropertyAnimatorCoreEditor/Private/Widgets/PropertyAnimatorCoreEditorEditPanelOptions.cpp
@@ -0,0 +1,16 @@
+// Copyright Epic Games, Inc. All Rights Reserved.
+
+#include "Widgets/PropertyAnimatorCoreEditorEditPanelOptions.h"
+
+#include "GameFramework/Actor.h"
+#include "Widgets/SPropertyAnimatorCoreEditorEditPanel.h"
+
+void FPropertyAnimatorCoreEditorEditPanelOptions::SetContextActor(AActor* InActor)
+{
+	ContextActorWeak = InActor;
+
+	if (const TSharedPtr<SPropertyAnimatorCoreEditorEditPanel> EditPanel = EditPanelWeak.Pin())
+	{
+		EditPanel->Update();
+	}
+}
\ No newline at end of file
diff --git a/Engine/Plugins/Experimental/PropertyAnimatorCore/Source/PropertyAnimatorCoreEditor/Private/Widgets/SPropertyAnimatorCoreEditorEditPanel.cpp b/Engine/Plugins/Experimental/PropertyAnimatorCore/Source/PropertyAnimatorCoreEditor/Private/Widgets/SPropertyAnimatorCoreEditorEditPanel.cpp
new file mode 100644
--- /dev/null
+++ b/Engine/Plugins/Experimental/PropertyAnimatorCore/Source/PropertyAnimatorCoreEditor/Private/Widgets/SPropertyAnimatorCoreEditorEditPanel.cpp
@@ -0,0 +1,429 @@
+// Copyright Epic Games, Inc. All Rights Reserved.
+
+#include "Widgets/SPropertyAnimatorCoreEditorEditPanel.h"
+
+#include "DragDropOps/PropertyAnimatorCoreEditorViewDragDropOp.h"
+#include "Framework/Application/SlateApplication.h"
+#include "Framework/Commands/GenericCommands.h"
+#include "Menus/PropertyAnimatorCoreEditorMenu.h"
+#include "Menus/PropertyAnimatorCoreEditorMenuContext.h"
+#include "Styles/PropertyAnimatorCoreEditorStyle.h"
+#include "Styling/ToolBarStyle.h"
+#include "Subsystems/PropertyAnimatorCoreEditorSubsystem.h"
+#include "Subsystems/PropertyAnimatorCoreSubsystem.h"
+#include "ToolMenu.h"
+#include "ToolMenus.h"
+#include "Views/SPropertyAnimatorCoreEditorControllersView.h"
+#include "Views/SPropertyAnimatorCoreEditorPropertiesView.h"
+#include "Widgets/Input/SCheckBox.h"
+#include "Widgets/Layout/SWidgetSwitcher.h"
+#include "Widgets/SWindow.h"
+
+#define LOCTEXT_NAMESPACE "PropertyAnimatorCoreEditorEditPanel"
+
+SPropertyAnimatorCoreEditorEditPanel::~SPropertyAnimatorCoreEditorEditPanel()
+{
+	UPropertyAnimatorCoreBase::OnAnimatorCreatedDelegate.RemoveAll(this);
+	UPropertyAnimatorCoreBase::OnAnimatorRemovedDelegate.RemoveAll(this);
+	UPropertyAnimatorCoreBase::OnAnimatorRenamedDelegate.RemoveAll(this);
+	UPropertyAnimatorCoreBase::OnAnimatorPropertyLinkedDelegate.RemoveAll(this);
+	UPropertyAnimatorCoreBase::OnAnimatorPropertyUnlinkedDelegate.RemoveAll(this);
+}
+
+TSharedPtr<SPropertyAnimatorCoreEditorEditPanel> SPropertyAnimatorCoreEditorEditPanel::OpenWindow()
+{
+	TSharedPtr<SPropertyAnimatorCoreEditorEditPanel> PropertyControllerPanel = SNew(SPropertyAnimatorCoreEditorEditPanel);
+
+	const TAttribute<FText> WindowTitle = TAttribute<FText>::CreateSP(PropertyControllerPanel.Get(), &SPropertyAnimatorCoreEditorEditPanel::GetTitle);
+
+	const TSharedRef<SWindow> PropertyControlWindow =
+		SNew(SWindow)
+		.Title(WindowTitle)
+		.ClientSize(FVector2f(600.0f, 600.0f))
+		.SizingRule(ESizingRule::FixedSize)
+		.SupportsMaximize(false)
+		.SupportsMinimize(false)
+		.FocusWhenFirstShown(true)
+		.Content()
+		[
+			PropertyControllerPanel.ToSharedRef()
+		];
+
+	FSlateApplication::Get().AddWindow(PropertyControlWindow, true);
+
+	return PropertyControllerPanel;
+}
+
+void SPropertyAnimatorCoreEditorEditPanel::CloseWindow() const
+{
+	if (const TSharedPtr<SWindow> Window = FSlateApplication::Get().FindWidgetWindow(AsShared()))
+	{
+		Window->RequestDestroyWindow();
+	}
+}
+
+void SPropertyAnimatorCoreEditorEditPanel::FocusWindow() const
+{
+	if (const TSharedPtr<SWindow> Window = FSlateApplication::Get().FindWidgetWindow(AsShared()))
+	{
+		Window->BringToFront();
+	}
+}
+
+void SPropertyAnimatorCoreEditorEditPanel::Construct(const FArguments& InArgs)
+{
+	static const FSlateIcon ControllersIcon(FPropertyAnimatorCoreEditorStyle::Get().GetStyleSetName(), "PropertyControlIcon.Default");
+	static const FSlateIcon PropertiesIcon(FPropertyAnimatorCoreEditorStyle::Get().GetStyleSetName(), "PropertyControlIcon.Default");
+
+	static const FToolBarStyle& ToolBarStyle = FAppStyle::Get().GetWidgetStyle<FToolBarStyle>("SlimToolBar");
+	static const FCheckBoxStyle* const CheckStyle = &ToolBarStyle.ToggleButton;
+
+	TSharedPtr<SPropertyAnimatorCoreEditorEditPanel> This = SharedThis(this);
+
+	Options = FPropertyAnimatorCoreEditorEditPanelOptions(This);
+
+	CommandList = MakeShared<FUICommandList>();
+
+	const FGenericCommands& GenericCommands = FGenericCommands::Get();
+
+	CommandList->MapAction(GenericCommands.Delete
+		, FExecuteAction::CreateSP(this, &SPropertyAnimatorCoreEditorEditPanel::ExecuteDeleteCommand)
+		, FCanExecuteAction::CreateSP(this, &SPropertyAnimatorCoreEditorEditPanel::CanExecuteDeleteCommand));
+
+	CommandList->MapAction(GenericCommands.Rename
+		, FExecuteAction::CreateSP(this, &SPropertyAnimatorCoreEditorEditPanel::ExecuteRenameCommand)
+		, FCanExecuteAction::CreateSP(this, &SPropertyAnimatorCoreEditorEditPanel::CanExecuteRenameCommand));
+
+	ChildSlot
+	.HAlign(HAlign_Fill)
+	.VAlign(VAlign_Fill)
+	[
+		SNew(SVerticalBox)
+		// Toolbar to switch between different views
+		+ SVerticalBox::Slot()
+		.AutoHeight()
+		[
+			SNew(SBorder)
+			.BorderImage(&ToolBarStyle.SeparatorBrush)
+			.HAlign(HAlign_Center)
+			.VAlign(VAlign_Center)
+			.Padding(0.f, 5.f)
+			[
+				SAssignNew(ViewsToolbar, SScrollBox)
+				.Orientation(EOrientation::Orient_Horizontal)
+				.ScrollBarThickness(FVector2D(1.f))
+				// Add "View by properties" button
+				+ SScrollBox::Slot()
+				.Padding(5.f, 0.f)
+				[
+					SNew(SCheckBox)
+					.Style(CheckStyle)
+					.ForegroundColor(FLinearColor::White)
+					.OnCheckStateChanged(this, &SPropertyAnimatorCoreEditorEditPanel::OnViewButtonClicked, SPropertyAnimatorCoreEditorEditPanel::PropertiesViewIndex)
+					.IsChecked(this, &SPropertyAnimatorCoreEditorEditPanel::IsViewButtonActive, SPropertyAnimatorCoreEditorEditPanel::PropertiesViewIndex)
+					[
+						SNew(SHorizontalBox)
+						+ SHorizontalBox::Slot()
+						.AutoWidth()
+						.Padding(5.f)
+						.HAlign(HAlign_Left)
+						.VAlign(VAlign_Center)
+						[
+							SNew(SImage)
+							.DesiredSizeOverride(FVector2D(16.f, 16.f))
+							.Image(PropertiesIcon.GetIcon())
+						]
+						+ SHorizontalBox::Slot()
+						.FillWidth(1.f)
+						.Padding(5.f)
+						.HAlign(HAlign_Left)
+						.VAlign(VAlign_Center)
+						[
+							SNew(STextBlock)
+							.Text(LOCTEXT("PropertiesView.Label", "View by properties"))
+						]
+					]
+				]
+				// Add "View by controllers" button
+				+ SScrollBox::Slot()
+				.Padding(5.f, 0.f)
+				[
+					SNew(SCheckBox)
+					.Style(CheckStyle)
+					.ForegroundColor(FLinearColor::White)
+					.OnCheckStateChanged(this, &SPropertyAnimatorCoreEditorEditPanel::OnViewButtonClicked, SPropertyAnimatorCoreEditorEditPanel::ControllersViewIndex)
+					.IsChecked(this, &SPropertyAnimatorCoreEditorEditPanel::IsViewButtonActive, SPropertyAnimatorCoreEditorEditPanel::ControllersViewIndex)
+					[
+						SNew(SHorizontalBox)
+						+ SHorizontalBox::Slot()
+						.AutoWidth()
+						.Padding(5.f)
+						.HAlign(HAlign_Left)
+						.VAlign(VAlign_Center)
+						[
+							SNew(SImage)
+							.DesiredSizeOverride(FVector2D(16.f, 16.f))
+							.Image(ControllersIcon.GetIcon())
+						]
+						+ SHorizontalBox::Slot()
+						.FillWidth(1.f)
+						.Padding(5.f)
+						.HAlign(HAlign_Left)
+						.VAlign(VAlign_Center)
+						[
+							SNew(STextBlock)
+							.Text(LOCTEXT("ControllersView.Label", "View by animators"))
+						]
+					]
+				]
+			]
+		]
+		// Widget for each views
+		+ SVerticalBox::Slot()
+		.FillHeight(1.f)
+		[
+			SAssignNew(ViewsSwitcher, SWidgetSwitcher)
+			.WidgetIndex(ActiveViewIndex)
+			+ SWidgetSwitcher::Slot()
+			.Padding(5.f)
+			[
+				SAssignNew(PropertiesView, SPropertyAnimatorCoreEditorPropertiesView, This)
+			]
+			+ SWidgetSwitcher::Slot()
+			.Padding(5.f)
+			[
+				SAssignNew(ControllersView, SPropertyAnimatorCoreEditorControllersView, This)
+			]
+		]
+	];
+
+	UPropertyAnimatorCoreBase::OnAnimatorCreatedDelegate.AddSP(this, &SPropertyAnimatorCoreEditorEditPanel::OnControllerUpdated);
+	UPropertyAnimatorCoreBase::OnAnimatorRemovedDelegate.AddSP(this, &SPropertyAnimatorCoreEditorEditPanel::OnControllerUpdated);
+	UPropertyAnimatorCoreBase::OnAnimatorRenamedDelegate.AddSP(this, &SPropertyAnimatorCoreEditorEditPanel::OnControllerUpdated);
+	UPropertyAnimatorCoreBase::OnAnimatorPropertyLinkedDelegate.AddSP(this, &SPropertyAnimatorCoreEditorEditPanel::OnControllerPropertyUpdated);
+	UPropertyAnimatorCoreBase::OnAnimatorPropertyUnlinkedDelegate.AddSP(this, &SPropertyAnimatorCoreEditorEditPanel::OnControllerPropertyUpdated);
+}
+
+FReply SPropertyAnimatorCoreEditorEditPanel::OnDragStart(const FGeometry& InGeometry, const FPointerEvent& InPointerEvent)
+{
+	if (!GlobalSelection.IsEmpty())
+	{
+		const TSharedRef<FPropertyAnimatorCoreEditorViewDragDropOp> DragDropOp = FPropertyAnimatorCoreEditorViewDragDropOp::New(GlobalSelection);
+		return FReply::Handled().BeginDragDrop(DragDropOp);
+	}
+
+	return FReply::Unhandled();
+}
+
+FReply SPropertyAnimatorCoreEditorEditPanel::OnMouseButtonDown(const FGeometry& InGeometry, const FPointerEvent& InMouseEvent)
+{
+	if (InMouseEvent.GetEffectingButton() == EKeys::RightMouseButton)
+	{
+		// Open context menu
+		const FVector2D MenuLocation = FSlateApplication::Get().GetCursorPos();
+
+		FSlateApplication::Get().PushMenu(
+			SharedThis(this),
+			FWidgetPath(),
+			GenerateContextMenuWidget(),
+			MenuLocation,
+			FPopupTransitionEffect(FPopupTransitionEffect::ContextMenu)
+		);
+
+		return FReply::Handled();
+	}
+
+	return SCompoundWidget::OnMouseButtonDown(InGeometry, InMouseEvent);
+}
+
+FReply SPropertyAnimatorCoreEditorEditPanel::OnKeyDown(const FGeometry& InGeometry, const FKeyEvent& InKeyEvent)
+{
+	if (CommandList.IsValid()
+		&& CommandList->ProcessCommandBindings(InKeyEvent))
+	{
+		return FReply::Handled();
+	}
+
+	return SCompoundWidget::OnKeyDown(InGeometry, InKeyEvent);
+}
+
+TSharedRef<SWidget> SPropertyAnimatorCoreEditorEditPanel::GenerateContextMenuWidget()
+{
+	UToolMenus* Menus = UToolMenus::Get();
+
+	check(Menus);
+
+	static const FName EditPanelContextMenuName = TEXT("AnimatorEditPanelContextMenu");
+
+	if (!Menus->IsMenuRegistered(EditPanelContextMenuName))
+	{
+		UToolMenu* ContextMenu = Menus->RegisterMenu(EditPanelContextMenuName, NAME_None, EMultiBoxType::Menu);
+		ContextMenu->AddDynamicSection(TEXT("FillAnimatorEditPanelContextSection"), FNewToolMenuDelegate::CreateSP(this, &SPropertyAnimatorCoreEditorEditPanel::FillAnimatorContextSection));
+
+		const FGenericCommands& GenericCommands = FGenericCommands::Get();
+		FToolMenuSection& GenericActionsSection = ContextMenu->FindOrAddSection(TEXT("GenericActions"), LOCTEXT("GenericActions.Label", "Generic Actions"));
+		GenericActionsSection.AddMenuEntry(GenericCommands.Delete);
+		GenericActionsSection.AddMenuEntry(GenericCommands.Rename);
+	}
+
+	UPropertyAnimatorCoreEditorMenuContext* MenuContext = nullptr;
+	if (ActiveViewIndex == SPropertyAnimatorCoreEditorEditPanel::ControllersViewIndex && GlobalSelection.IsEmpty())
+	{
+		MenuContext = NewObject<UPropertyAnimatorCoreEditorMenuContext>();
+		MenuContext->SetPropertyData(FPropertyAnimatorCoreData(GetOptions().GetContextActor(), nullptr, nullptr));
+	}
+
+	const FToolMenuContext ToolMenuContext(CommandList, nullptr, MenuContext);
+	return Menus->GenerateWidget(EditPanelContextMenuName, ToolMenuContext);
+}
+
+void SPropertyAnimatorCoreEditorEditPanel::FillAnimatorContextSection(UToolMenu* InToolMenu) const
+{
+	if (!InToolMenu)
+	{
+		return;
+	}
+
+	const UPropertyAnimatorCoreEditorMenuContext* Context = InToolMenu->FindContext<UPropertyAnimatorCoreEditorMenuContext>();
+
+	if (!Context)
+	{
+		return;
+	}
+
+	UPropertyAnimatorCoreEditorSubsystem* EditorSubsystem = UPropertyAnimatorCoreEditorSubsystem::Get();
+	if (!EditorSubsystem)
+	{
+		return;
+	}
+
+	const FPropertyAnimatorCoreEditorMenuContext MenuContext({}, {Context->GetPropertyData()});
+	const FPropertyAnimatorCoreEditorMenuOptions MenuOptions({EPropertyAnimatorCoreEditorMenuType::New});
+	EditorSubsystem->FillAnimatorMenu(InToolMenu, MenuContext, MenuOptions);
+}
+
+bool SPropertyAnimatorCoreEditorEditPanel::CanExecuteDeleteCommand() const
+{
+	return !GlobalSelection.IsEmpty();
+}
+
+void SPropertyAnimatorCoreEditorEditPanel::ExecuteDeleteCommand()
+{
+	if (!CanExecuteDeleteCommand())
+	{
+		return;
+	}
+
+	const UPropertyAnimatorCoreSubsystem* ControllerSubsystem = UPropertyAnimatorCoreSubsystem::Get();
+	for (const FPropertiesViewControllerItem& Item : GlobalSelection)
+	{
+		UPropertyAnimatorCoreBase* Controller = Item.ControllerWeak.Get();
+
+		if (!Controller)
+		{
+			continue;
+		}
+
+		if (Item.Property.IsValid())
+		{
+			Controller->UnlinkProperty(*Item.Property);
+		}
+		else
+		{
+			ControllerSubsystem->RemoveAnimator(Controller);
+		}
+	}
+
+	GlobalSelection.Empty();
+}
+
+bool SPropertyAnimatorCoreEditorEditPanel::CanExecuteRenameCommand() const
+{
+	if (GlobalSelection.Num() == 1 && ActiveViewIndex == SPropertyAnimatorCoreEditorEditPanel::ControllersViewIndex)
+	{
+		const FPropertiesViewControllerItem Item = GlobalSelection.Array()[0];
+
+		if (Item.ControllerWeak.IsValid())
+		{
+			return true;
+		}
+	}
+
+	return false;
+}
+
+void SPropertyAnimatorCoreEditorEditPanel::ExecuteRenameCommand()
+{
+	if (!CanExecuteRenameCommand())
+	{
+		return;
+	}
+
+	const FPropertiesViewControllerItem Item = GlobalSelection.Array()[0];
+	UPropertyAnimatorCoreBase* Controller = Item.ControllerWeak.Get();
+
+	if (!Controller)
+	{
+		return;
+	}
+
+	OnControllerRenameRequestedDelegate.Broadcast(Controller);
+}
+
+void SPropertyAnimatorCoreEditorEditPanel::Update() const
+{
+	if (PropertiesView.IsValid())
+	{
+		PropertiesView->Update();
+	}
+
+	if (ControllersView.IsValid())
+    {
+    	ControllersView->Update();
+    }
+}
+
+FText SPropertyAnimatorCoreEditorEditPanel::GetTitle() const
+{
+	static const FText TitleText = LOCTEXT("EditAnimatorsWindow.Title", "Edit Animators ({0})");
+	const AActor* OwningActor = GetOptions().GetContextActor();
+	return FText::Format(TitleText, FText::FromString(OwningActor ? OwningActor->GetActorNameOrLabel() : TEXT("Invalid Actor")));
+}
+
+void SPropertyAnimatorCoreEditorEditPanel::OnViewButtonClicked(ECheckBoxState InCheckBoxState, int32 InViewIdx)
+{
+	ActiveViewIndex = InViewIdx;
+
+	GlobalSelection.Empty();
+	OnGlobalSelectionChangedDelegate.Broadcast();
+
+	if (ViewsSwitcher.IsValid())
+	{
+		ViewsSwitcher->SetActiveWidgetIndex(ActiveViewIndex);
+	}
+}
+
+ECheckBoxState SPropertyAnimatorCoreEditorEditPanel::IsViewButtonActive(int32 InViewIdx) const
+{
+	return InViewIdx == ActiveViewIndex
+		? ECheckBoxState::Checked
+		: ECheckBoxState::Unchecked;
+}
+
+void SPropertyAnimatorCoreEditorEditPanel::OnControllerUpdated(UPropertyAnimatorCoreBase* InController)
+{
+	if (InController && InController->GetAnimatorActor() == GetOptions().GetContextActor())
+	{
+		Update();
+	}
+}
+
+void SPropertyAnimatorCoreEditorEditPanel::OnControllerPropertyUpdated(UPropertyAnimatorCoreBase* InController, const FPropertyAnimatorCoreData& InProperty)
+{
+	if (InController && InController->GetAnimatorActor() == GetOptions().GetContextActor())
+	{
+		Update();
+	}
+}
+
+#undef LOCTEXT_NAMESPACE
diff --git a/Engine/Plugins/Experimental/PropertyAnimatorCore/Source/PropertyAnimatorCoreEditor/Private/Widgets/SPropertyAnimatorCoreEditorEditPanel.h b/Engine/Plugins/Experimental/PropertyAnimatorCore/Source/PropertyAnimatorCoreEditor/Private/Widgets/SPropertyAnimatorCoreEditorEditPanel.h
new file mode 100644
--- /dev/null
+++ b/Engine/Plugins/Experimental/PropertyAnimatorCore/Source/PropertyAnimatorCoreEditor/Private/Widgets/SPropertyAnimatorCoreEditorEditPanel.h
@@ -0,0 +1,110 @@
+// Copyright Epic Games, Inc. All Rights Reserved.
+
+#pragma once
+
+#include "Framework/Commands/UICommandList.h"
+#include "Templates/SharedPointer.h"
+#include "UObject/WeakObjectPtr.h"
+#include "Widgets/Layout/SScrollBox.h"
+#include "Widgets/Layout/SWidgetSwitcher.h"
+#include "Widgets/PropertyAnimatorCoreEditorEditPanelOptions.h"
+#include "Widgets/Views/PropertyAnimatorCoreEditorViewItem.h"
+#include "Widgets/Views/STreeView.h"
+#include "Widgets/Views/SHeaderRow.h"
+#include "Widgets/SCompoundWidget.h"
+
+class SPropertyAnimatorCoreEditorControllersView;
+class SPropertyAnimatorCoreEditorPropertiesView;
+class UToolMenu;
+
+/** Window content showing the link between controller and properties in various views */
+class SPropertyAnimatorCoreEditorEditPanel : public SCompoundWidget
+{
+	friend class SPropertyAnimatorCoreEditorControllersViewTableRow;
+	friend class SPropertyAnimatorCoreEditorPropertiesViewControllerTableRow;
+
+public:
+	static inline constexpr int32 PropertiesViewIndex = 0;
+	static inline constexpr int32 ControllersViewIndex = 1;
+	static inline const FName HeaderPropertyColumnName = "Properties";
+	static inline const FName HeaderAnimatorColumnName = "Animators";
+	static inline const FName HeaderActionColumnName = "Actions";
+
+	DECLARE_MULTICAST_DELEGATE(FOnGlobalSelectionChanged)
+	FOnGlobalSelectionChanged OnGlobalSelectionChangedDelegate;
+
+	DECLARE_MULTICAST_DELEGATE_OneParam(FOnControllerRenameRequested, UPropertyAnimatorCoreBase* /** Controller */)
+	FOnControllerRenameRequested OnControllerRenameRequestedDelegate;
+
+	SLATE_BEGIN_ARGS(SPropertyAnimatorCoreEditorEditPanel) {}
+	SLATE_END_ARGS()
+
+	virtual ~SPropertyAnimatorCoreEditorEditPanel() override;
+
+	static TSharedPtr<SPropertyAnimatorCoreEditorEditPanel> OpenWindow();
+
+	void CloseWindow() const;
+
+	void FocusWindow() const;
+
+	void Construct(const FArguments& InArgs);
+
+	void Update() const;
+
+	FPropertyAnimatorCoreEditorEditPanelOptions& GetOptions()
+	{
+		return Options;
+	}
+
+	const FPropertyAnimatorCoreEditorEditPanelOptions& GetOptions() const
+	{
+		return Options;
+	}
+
+	TSet<FPropertiesViewControllerItem>& GetGlobalSelection()
+	{
+		return GlobalSelection;
+	}
+
+protected:
+	FReply OnDragStart(const FGeometry& InGeometry, const FPointerEvent& InPointerEvent);
+	virtual FReply OnMouseButtonDown(const FGeometry& InGeometry, const FPointerEvent& InMouseEvent) override;
+	virtual FReply OnKeyDown(const FGeometry& InGeometry, const FKeyEvent& InKeyEvent) override;
+
+	TSharedRef<SWidget> GenerateContextMenuWidget();
+	void FillAnimatorContextSection(UToolMenu* InToolMenu) const;
+
+	void OnViewButtonClicked(ECheckBoxState InCheckBoxState, int32 InViewIdx);
+	ECheckBoxState IsViewButtonActive(int32 InViewIdx) const;
+	FText GetTitle() const;
+
+	void OnControllerUpdated(UPropertyAnimatorCoreBase* InController);
+	void OnControllerPropertyUpdated(UPropertyAnimatorCoreBase* InController, const FPropertyAnimatorCoreData& InProperty);
+
+	bool CanExecuteDeleteCommand() const;
+	void ExecuteDeleteCommand();
+
+	bool CanExecuteRenameCommand() const;
+	void ExecuteRenameCommand();
+
+	/** Contains the available button views */
+	TSharedPtr<SScrollBox> ViewsToolbar;
+
+	/** Contains the different views widget */
+	TSharedPtr<SWidgetSwitcher> ViewsSwitcher;
+
+	/** Properties view widget */
+	TSharedPtr<SPropertyAnimatorCoreEditorPropertiesView> PropertiesView;
+
+	/** Controllers view widget */
+	TSharedPtr<SPropertyAnimatorCoreEditorControllersView> ControllersView;
+
+	int32 ActiveViewIndex = SPropertyAnimatorCoreEditorEditPanel::PropertiesViewIndex;
+
+	FPropertyAnimatorCoreEditorEditPanelOptions Options;
+
+	/** Selection on each views to apply commands */
+	TSet<FPropertiesViewControllerItem> GlobalSelection;
+
+	TSharedPtr<FUICommandList> CommandList;
+};
diff --git a/Engine/Plugins/Experimental/PropertyAnimatorCore/Source/PropertyAnimatorCoreEditor/Private/Widgets/TableRows/SPropertyAnimatorCoreEditorControllersViewTableRow.cpp b/Engine/Plugins/Experimental/PropertyAnimatorCore/Source/PropertyAnimatorCoreEditor/Private/Widgets/TableRows/SPropertyAnimatorCoreEditorControllersViewTableRow.cpp
new file mode 100644
--- /dev/null
+++ b/Engine/Plugins/Experimental/PropertyAnimatorCore/Source/PropertyAnimatorCoreEditor/Private/Widgets/TableRows/SPropertyAnimatorCoreEditorControllersViewTableRow.cpp
@@ -0,0 +1,389 @@
+// Copyright Epic Games, Inc. All Rights Reserved.
+
+#include "Widgets/TableRows/SPropertyAnimatorCoreEditorControllersViewTableRow.h"
+
+#include "Framework/Application/SlateApplication.h"
+#include "Styling/SlateIconFinder.h"
+#include "Widgets/SPropertyAnimatorCoreEditorEditPanel.h"
+#include "Widgets/DragDropOps/PropertyAnimatorCoreEditorViewDragDropOp.h"
+#include "Widgets/Layout/SScaleBox.h"
+#include "Widgets/Views/SPropertyAnimatorCoreEditorControllersView.h"
+
+#define LOCTEXT_NAMESPACE "PropertyAnimatorCoreEditorControllerTableRow"
+
+SPropertyAnimatorCoreEditorControllersViewTableRow::~SPropertyAnimatorCoreEditorControllersViewTableRow()
+{
+	if (const TSharedPtr<SPropertyAnimatorCoreEditorControllersView> View = ViewWeak.Pin())
+	{
+		if (const TSharedPtr<SPropertyAnimatorCoreEditorEditPanel> EditPanel = View->GetEditPanel())
+		{
+			EditPanel->OnGlobalSelectionChangedDelegate.RemoveAll(this);
+		}
+	}
+}
+
+void SPropertyAnimatorCoreEditorControllersViewTableRow::Construct(const FArguments& InArgs, const TSharedRef<STableViewBase>& InOwnerTableView, TSharedPtr<SPropertyAnimatorCoreEditorControllersView> InView, FControllersViewItemPtr InItem)
+{
+	ViewWeak = InView;
+	RowItemWeak = InItem;
+
+	check(InItem.IsValid())
+
+	const TSharedPtr<SPropertyAnimatorCoreEditorEditPanel> EditPanel = InView->GetEditPanel();
+
+	check(EditPanel.IsValid())
+
+	EditPanel->OnGlobalSelectionChangedDelegate.AddSP(this, &SPropertyAnimatorCoreEditorControllersViewTableRow::OnGlobalSelectionChanged);
+	EditPanel->OnControllerRenameRequestedDelegate.AddSP(this, &SPropertyAnimatorCoreEditorControllersViewTableRow::OnControllerRenameRequested);
+
+	SMultiColumnTableRow<FControllersViewItemPtr>::Construct(
+		FSuperRowType::FArguments()
+		.Style(&FAppStyle::Get().GetWidgetStyle<FTableRowStyle>("SceneOutliner.TableViewRow"))
+		.Padding(0.f)
+		.ShowSelection(true)
+		.OnDragDetected(EditPanel.Get(), &SPropertyAnimatorCoreEditorEditPanel::OnDragStart)
+		.OnPaintDropIndicator(this, &SPropertyAnimatorCoreEditorControllersViewTableRow::OnPropertyPaintDropIndicator)
+		.OnDrop(this, &SPropertyAnimatorCoreEditorControllersViewTableRow::OnPropertyDrop)
+		.OnCanAcceptDrop(this, &SPropertyAnimatorCoreEditorControllersViewTableRow::OnPropertyCanAcceptDrop),
+		InOwnerTableView
+	);
+
+	SetBorderImage(TAttribute<const FSlateBrush*>(this, &SPropertyAnimatorCoreEditorControllersViewTableRow::GetBorder));
+}
+
+TSharedRef<SWidget> SPropertyAnimatorCoreEditorControllersViewTableRow::GenerateWidgetForColumn(const FName& InColumnName)
+{
+	const FControllersViewItemPtr Item = RowItemWeak.Pin();
+
+	if (!Item.IsValid())
+	{
+		return SNullWidget::NullWidget;
+	}
+
+	const UPropertyAnimatorCoreBase* Controller = Item->ControlledProperty.ControllerWeak.Get();
+
+	if (!Controller)
+	{
+		return SNullWidget::NullWidget;
+	}
+
+	TSharedRef<SBox> BoxWrapper = SNew(SBox)
+		.Visibility(EVisibility::Visible)
+		.VAlign(VAlign_Fill)
+		.HAlign(HAlign_Fill);
+
+	if (InColumnName == SPropertyAnimatorCoreEditorEditPanel::HeaderAnimatorColumnName)
+	{
+		const FText ControllerDisplayText = FText::FromString(Controller->GetAnimatorDisplayName());
+		const FText ControllerNameText = FText::FromName(Controller->GetAnimatorOriginalName());
+		const FSlateIcon ControllerIcon = FSlateIconFinder::FindIconForClass(Controller->GetClass());
+
+		BoxWrapper->SetContent(
+			SNew(SHorizontalBox)
+			.Visibility(!Item->ControlledProperty.Property.IsValid() ? EVisibility::Visible : EVisibility::Hidden)
+			+ SHorizontalBox::Slot()
+			.AutoWidth()
+			[
+				SAssignNew(ExpanderArrowWidget, SExpanderArrow, SharedThis(this))
+			]
+			+ SHorizontalBox::Slot()
+			.HAlign(HAlign_Fill)
+			.VAlign(VAlign_Center)
+			.AutoWidth()
+			.Padding(5.f, 0.f)
+			[
+				SNew(SScaleBox)
+				.Stretch(EStretch::ScaleToFit)
+				[
+					SNew(SImage)
+					.Image(ControllerIcon.GetIcon())
+				]
+			]
+			+ SHorizontalBox::Slot()
+			.HAlign(HAlign_Fill)
+			.VAlign(VAlign_Fill)
+			.FillWidth(1.f)
+			.Padding(5.f, 0.f)
+			[
+				SNew(SVerticalBox)
+				.Visibility(EVisibility::Visible)
+				+ SVerticalBox::Slot()
+				.VAlign(VAlign_Center)
+				.HAlign(HAlign_Fill)
+				.Padding(0.f, 5.f)
+				.AutoHeight()
+				[
+					SAssignNew(ControllerEditableTextBox, SEditableTextBox)
+					.BackgroundColor(FSlateColor(FLinearColor::Transparent))
+					.Visibility(EVisibility::HitTestInvisible)
+					.Padding(0.f)
+					.Justification(ETextJustify::Left)
+					.Text(ControllerDisplayText)
+					.ToolTipText(ControllerDisplayText)
+					.IsReadOnly(true)
+					.RevertTextOnEscape(true)
+					.SelectAllTextWhenFocused(true)
+					.OnVerifyTextChanged(this, &SPropertyAnimatorCoreEditorControllersViewTableRow::OnVerifyControllerTextChanged)
+					.OnTextCommitted(this, &SPropertyAnimatorCoreEditorControllersViewTableRow::OnControllerTextCommitted)
+				]
+				+ SVerticalBox::Slot()
+				.VAlign(VAlign_Center)
+				.HAlign(HAlign_Fill)
+				.Padding(0.f, 5.f)
+				.AutoHeight()
+				[
+					SNew(STextBlock)
+					.Margin(0.f)
+					.Justification(ETextJustify::Left)
+					.Text(ControllerNameText)
+				]
+			]
+		);
+
+		// Handle renaming when double clicking
+		BoxWrapper->SetOnMouseDoubleClick(FPointerEventHandler::CreateSP(this, &SPropertyAnimatorCoreEditorControllersViewTableRow::OnControllerBoxDoubleClick));
+
+		// Remove border around editable box to make it look like a standard text box
+		ControllerEditableTextBox->SetBorderBackgroundColor(FSlateColor(FLinearColor::Transparent));
+		ControllerEditableTextBox->SetBorderImage(nullptr);
+	}
+	else if (InColumnName == SPropertyAnimatorCoreEditorEditPanel::HeaderPropertyColumnName)
+	{
+		const FText PropertyDisplayName = Item->ControlledProperty.Property.IsValid()
+			? FText::FromName(Item->ControlledProperty.Property->GetPropertyDisplayName())
+			: FText::Format(LOCTEXT("AnimatorPropertiesLinkedCount", "{0} Propertie(s) Linked"), Controller->GetLinkedPropertiesCount());
+
+		BoxWrapper->SetContent(
+			SNew(STextBlock)
+			.Text(PropertyDisplayName)
+			.ToolTipText(PropertyDisplayName)
+		);
+
+		BoxWrapper->SetVAlign(VAlign_Center);
+	}
+
+	return BoxWrapper;
+}
+
+int32 SPropertyAnimatorCoreEditorControllersViewTableRow::OnPropertyPaintDropIndicator(EItemDropZone InDropZone, const FPaintArgs& InPaintArgs, const FGeometry& InGeometry, const FSlateRect& InSlateRect, FSlateWindowElementList& OutElements, int32 InLayerIndex, const FWidgetStyle& InWidgetStyle, bool InbParentEnabled)
+{
+	const FSlateBrush* DropIndicatorBrush = GetDropIndicatorBrush(InDropZone);
+
+	constexpr float OffsetX = 10.0f;
+	const FVector2D Offset(OffsetX * GetIndentLevel(), 0.f);
+	FSlateDrawElement::MakeBox
+	(
+		OutElements,
+		InLayerIndex++,
+		InGeometry.ToPaintGeometry(FVector2D(InGeometry.GetLocalSize() - Offset), FSlateLayoutTransform(Offset)),
+		DropIndicatorBrush,
+		ESlateDrawEffect::None,
+		DropIndicatorBrush->GetTint(InWidgetStyle) * InWidgetStyle.GetColorAndOpacityTint()
+	);
+
+	return InLayerIndex;
+}
+
+FReply SPropertyAnimatorCoreEditorControllersViewTableRow::OnPropertyDrop(FDragDropEvent const& InDragDropEvent)
+{
+	if (const TSharedPtr<FPropertyAnimatorCoreEditorViewDragDropOp> DragDropOp = InDragDropEvent.GetOperationAs<FPropertyAnimatorCoreEditorViewDragDropOp>())
+	{
+		if (const FControllersViewItemPtr RowItem = RowItemWeak.Pin())
+		{
+			if (UPropertyAnimatorCoreBase* Controller = RowItem->ControlledProperty.ControllerWeak.Get())
+			{
+				const bool bIsControlPressed = FSlateApplication::Get().GetModifierKeys().IsControlDown();
+
+				for (const FPropertiesViewControllerItem& DraggedItem : DragDropOp->GetDraggedItems())
+				{
+					if (!DraggedItem.Property.IsValid())
+					{
+						continue;
+					}
+
+					// Unlink from current controller and link to new
+					UPropertyAnimatorCoreBase* PreviousController = DraggedItem.ControllerWeak.Get();
+					if (!bIsControlPressed && PreviousController)
+				    {
+				    	PreviousController->UnlinkProperty(*DraggedItem.Property);
+				    }
+
+					Controller->LinkProperty(*DraggedItem.Property);
+				}
+			}
+		}
+
+		return FReply::Handled().EndDragDrop();
+	}
+
+	return FReply::Unhandled();
+}
+
+TOptional<EItemDropZone> SPropertyAnimatorCoreEditorControllersViewTableRow::OnPropertyCanAcceptDrop(const FDragDropEvent& InDragDropEvent, EItemDropZone InDropZone, FControllersViewItemPtr InItem)
+{
+	// Get dragged properties and check that the controller we drop them on supports them
+	if (const TSharedPtr<FPropertyAnimatorCoreEditorViewDragDropOp> DragDropOp = InDragDropEvent.GetOperationAs<FPropertyAnimatorCoreEditorViewDragDropOp>())
+	{
+		const UPropertyAnimatorCoreBase* Controller = InItem->ControlledProperty.ControllerWeak.Get();
+
+		if (InItem.IsValid()
+			&& !InItem->ControlledProperty.Property.IsValid()
+			&& Controller)
+		{
+			for (const FPropertiesViewControllerItem& DraggedItem : DragDropOp->GetDraggedItems())
+			{
+				if (const FPropertyAnimatorCoreData* Property = DraggedItem.Property.Get())
+				{
+					if (Controller->IsPropertySupported(*Property)
+						&& !Controller->IsPropertyLinked(*Property))
+					{
+						return EItemDropZone::OntoItem;
+					}
+				}
+			}
+		}
+	}
+
+	return TOptional<EItemDropZone>();
+}
+
+void SPropertyAnimatorCoreEditorControllersViewTableRow::OnGlobalSelectionChanged()
+{
+	const TSharedPtr<SPropertyAnimatorCoreEditorControllersView> ControllersView = ViewWeak.Pin();
+
+	if (!ControllersView.IsValid())
+	{
+		return;
+	}
+
+	const TSharedPtr<SPropertyAnimatorCoreEditorEditPanel> EditPanel = ControllersView->GetEditPanel();
+
+	if (!EditPanel.IsValid())
+	{
+		return;
+	}
+
+	const FControllersViewItemPtr RowItem = RowItemWeak.Pin();
+
+	if (!RowItem.IsValid())
+	{
+		return;
+	}
+
+	const TSet<FPropertiesViewControllerItem>& GlobalSelection = EditPanel->GetGlobalSelection();
+
+	if (GlobalSelection.Contains(RowItem->ControlledProperty))
+	{
+		ControllersView->ControllersTree->SetItemSelection(RowItem, true, ESelectInfo::Direct);
+	}
+	else
+	{
+		ControllersView->ControllersTree->SetItemSelection(RowItem, false, ESelectInfo::Direct);
+	}
+}
+
+void SPropertyAnimatorCoreEditorControllersViewTableRow::OnControllerRenameRequested(UPropertyAnimatorCoreBase* InController)
+{
+	const FControllersViewItemPtr ItemRow = RowItemWeak.Pin();
+	if (!ItemRow.IsValid())
+	{
+		return;
+	}
+
+	if (InController != ItemRow->ControlledProperty.ControllerWeak.Get())
+	{
+		return;
+	}
+
+	BeginRenamingOperation();
+}
+
+bool SPropertyAnimatorCoreEditorControllersViewTableRow::OnVerifyControllerTextChanged(const FText& InText, FText& OutError) const
+{
+	if (!ControllerEditableTextBox.IsValid() || ControllerEditableTextBox->IsReadOnly() || !RowItemWeak.IsValid())
+	{
+		OutError = LOCTEXT("ControllerTextInvalidRow", "Invalid row item");
+		return false;
+	}
+
+	if (InText.IsEmpty())
+	{
+		OutError = LOCTEXT("ControllerTextInvalidName", "Invalid animator name");
+		return false;
+	}
+
+	// Add more rules here
+
+	return true;
+}
+
+void SPropertyAnimatorCoreEditorControllersViewTableRow::OnControllerTextCommitted(const FText& InText, ETextCommit::Type InType) const
+{
+	if (!ControllerEditableTextBox.IsValid() || ControllerEditableTextBox->IsReadOnly())
+	{
+		return;
+	}
+
+	const FControllersViewItemPtr ItemRow = RowItemWeak.Pin();
+	if (!ItemRow.IsValid())
+	{
+		return;
+	}
+
+	UPropertyAnimatorCoreBase* Controller = ItemRow->ControlledProperty.ControllerWeak.Get();
+	if (!Controller || Controller->IsTemplate())
+	{
+		return;
+	}
+
+	EndRenamingOperation();
+
+	if (InType == ETextCommit::OnEnter)
+	{
+		const FName NewControllerDisplayName(InText.ToString());
+		Controller->SetAnimatorDisplayName(NewControllerDisplayName);
+		ControllerEditableTextBox->SetText(InText);
+	}
+	else
+	{
+		const FText OriginalDisplayNameText = FText::FromString(Controller->GetAnimatorDisplayName());
+		ControllerEditableTextBox->SetText(OriginalDisplayNameText);
+	}
+}
+
+FReply SPropertyAnimatorCoreEditorControllersViewTableRow::OnControllerBoxDoubleClick(const FGeometry& InGeometry, const FPointerEvent& InEvent) const
+{
+	BeginRenamingOperation();
+	return FReply::Handled();
+}
+
+void SPropertyAnimatorCoreEditorControllersViewTableRow::BeginRenamingOperation() const
+{
+	if (!ControllerEditableTextBox.IsValid() || !ControllerEditableTextBox->IsReadOnly())
+	{
+		return;
+	}
+
+	const FEditableTextBoxStyle& EditableStyle = FCoreStyle::Get().GetWidgetStyle<FEditableTextBoxStyle>("NormalEditableTextBox");
+
+	ControllerEditableTextBox->SetBorderBackgroundColor(FLinearColor::White);
+	ControllerEditableTextBox->SetBorderImage(&EditableStyle.BackgroundImageFocused);
+	ControllerEditableTextBox->SetIsReadOnly(false);
+	ControllerEditableTextBox->SetVisibility(EVisibility::Visible);
+
+	FSlateApplication::Get().SetKeyboardFocus(ControllerEditableTextBox);
+}
+
+void SPropertyAnimatorCoreEditorControllersViewTableRow::EndRenamingOperation() const
+{
+	if (!ControllerEditableTextBox.IsValid() || ControllerEditableTextBox->IsReadOnly())
+	{
+		return;
+	}
+
+	ControllerEditableTextBox->SetBorderImage(nullptr);
+	ControllerEditableTextBox->SetVisibility(EVisibility::HitTestInvisible);
+	ControllerEditableTextBox->SetIsReadOnly(true);
+}
+
+#undef LOCTEXT_NAMESPACE
\ No newline at end of file
diff --git a/Engine/Plugins/Experimental/PropertyAnimatorCore/Source/PropertyAnimatorCoreEditor/Private/Widgets/TableRows/SPropertyAnimatorCoreEditorControllersViewTableRow.h b/Engine/Plugins/Experimental/PropertyAnimatorCore/Source/PropertyAnimatorCoreEditor/Private/Widgets/TableRows/SPropertyAnimatorCoreEditorControllersViewTableRow.h
new file mode 100644
--- /dev/null
+++ b/Engine/Plugins/Experimental/PropertyAnimatorCore/Source/PropertyAnimatorCoreEditor/Private/Widgets/TableRows/SPropertyAnimatorCoreEditorControllersViewTableRow.h
@@ -0,0 +1,53 @@
+// Copyright Epic Games, Inc. All Rights Reserved.
+
+#pragma once
+
+#include "Templates/SharedPointer.h"
+#include "Widgets/Input/SEditableTextBox.h"
+#include "Widgets/Views/SPropertyAnimatorCoreEditorPropertiesView.h"
+#include "Widgets/Views/STableRow.h"
+#include "Widgets/Views/STreeView.h"
+
+class SPropertyAnimatorCoreEditorControllersView;
+
+class SPropertyAnimatorCoreEditorControllersViewTableRow : public SMultiColumnTableRow<FControllersViewItemPtr>
+{
+public:
+	SLATE_BEGIN_ARGS(SPropertyAnimatorCoreEditorControllersViewTableRow) {}
+	SLATE_END_ARGS()
+
+	virtual ~SPropertyAnimatorCoreEditorControllersViewTableRow() override;
+
+	void Construct(const FArguments& InArgs, const TSharedRef<STableViewBase>& InOwnerTableView,
+		TSharedPtr<SPropertyAnimatorCoreEditorControllersView> InView, FControllersViewItemPtr InItem);
+
+	virtual TSharedRef<SWidget> GenerateWidgetForColumn(const FName& InColumnName) override;
+
+protected:
+	int32 OnPropertyPaintDropIndicator(EItemDropZone InDropZone, const FPaintArgs& InPaintArgs, const FGeometry& InGeometry, const FSlateRect& InSlateRect, FSlateWindowElementList& OutElements, int32 InLayerIndex, const FWidgetStyle& InWidgetStyle, bool InbParentEnabled);
+
+	FReply OnPropertyDrop(FDragDropEvent const& InDragDropEvent);
+
+	TOptional<EItemDropZone> OnPropertyCanAcceptDrop(const FDragDropEvent& InDragDropEvent, EItemDropZone InDropZone, FControllersViewItemPtr InItem);
+
+	void OnGlobalSelectionChanged();
+
+	void OnControllerRenameRequested(UPropertyAnimatorCoreBase* InController);
+
+	/** Before commiting new name, place verify rules here */
+	bool OnVerifyControllerTextChanged(const FText& InText, FText& OutError) const;
+
+	/** When we commit the editable text box */
+	void OnControllerTextCommitted(const FText& InText, ETextCommit::Type InType) const;
+
+	/** Used to enter renaming mode */
+	FReply OnControllerBoxDoubleClick(const FGeometry& InGeometry, const FPointerEvent& InEvent) const;
+
+	void BeginRenamingOperation() const;
+
+	void EndRenamingOperation() const;
+
+	TSharedPtr<SEditableTextBox> ControllerEditableTextBox;
+	TWeakPtr<SPropertyAnimatorCoreEditorControllersView> ViewWeak;
+	FControllersViewItemPtrWeak RowItemWeak;
+};
\ No newline at end of file
