[PropertyAnimator] Moved plugin

diff --git a/Engine/Plugins/Experimental/PropertyAnimatorCore/Source/PropertyAnimatorCoreEditor/Private/Widgets/TableRows/SPropertyAnimatorCoreEditorPropertiesViewControllerTableRow.cpp b/Engine/Plugins/Experimental/PropertyAnimatorCore/Source/PropertyAnimatorCoreEditor/Private/Widgets/TableRows/SPropertyAnimatorCoreEditorPropertiesViewControllerTableRow.cpp
new file mode 100644
--- /dev/null
+++ b/Engine/Plugins/Experimental/PropertyAnimatorCore/Source/PropertyAnimatorCoreEditor/Private/Widgets/TableRows/SPropertyAnimatorCoreEditorPropertiesViewControllerTableRow.cpp
@@ -0,0 +1,111 @@
+// Copyright Epic Games, Inc. All Rights Reserved.
+
+#include "Widgets/TableRows/SPropertyAnimatorCoreEditorPropertiesViewControllerTableRow.h"
+
+#include "Styling/SlateIconFinder.h"
+#include "Widgets/Layout/SScaleBox.h"
+#include "Widgets/SPropertyAnimatorCoreEditorEditPanel.h"
+#include "Widgets/TableRows/SPropertyAnimatorCoreEditorPropertiesViewTableRow.h"
+
+SPropertyAnimatorCoreEditorPropertiesViewControllerTableRow::~SPropertyAnimatorCoreEditorPropertiesViewControllerTableRow()
+{
+	if (const TSharedPtr<SPropertyAnimatorCoreEditorPropertiesViewTableRow> Row = RowWeak.Pin())
+	{
+		if (const TSharedPtr<SPropertyAnimatorCoreEditorPropertiesView> View = Row->ViewWeak.Pin())
+		{
+			if (const TSharedPtr<SPropertyAnimatorCoreEditorEditPanel> EditPanel = View->GetEditPanel())
+			{
+				EditPanel->OnGlobalSelectionChangedDelegate.RemoveAll(this);
+			}
+		}
+	}
+}
+
+void SPropertyAnimatorCoreEditorPropertiesViewControllerTableRow::Construct(const FArguments& InArgs, const TSharedRef<STableViewBase>& InOwnerTableView, TSharedPtr<SPropertyAnimatorCoreEditorPropertiesViewTableRow> InView, FPropertiesViewControllerItemPtr InItem)
+{
+	RowWeak = InView;
+	TileItemWeak = InItem;
+
+	check(InItem.IsValid())
+
+	const TSharedPtr<SPropertyAnimatorCoreEditorPropertiesView> View = InView->ViewWeak.Pin();
+
+	check(View.IsValid())
+
+	const TSharedPtr<SPropertyAnimatorCoreEditorEditPanel> EditPanel = View->GetEditPanel();
+
+	check(EditPanel.IsValid())
+
+	EditPanel->OnGlobalSelectionChangedDelegate.AddSP(this, &SPropertyAnimatorCoreEditorPropertiesViewControllerTableRow::OnGlobalSelectionChanged);
+
+	const UPropertyAnimatorCoreBase* Controller = InItem->ControllerWeak.Get();
+
+	check(Controller);
+
+	const FText ControllerDisplayName = FText::FromString(Controller->GetAnimatorDisplayName());
+	const FSlateIcon ControllerIcon = FSlateIconFinder::FindIconForClass(Controller->GetClass());
+
+	ChildSlot
+	.Padding(1.f)
+	[
+		SNew(SScaleBox)
+		.Stretch(EStretch::ScaleToFit)
+		[
+			SNew(SImage)
+			.ToolTipText(ControllerDisplayName)
+			.Image(ControllerIcon.GetIcon())
+		]
+	];
+
+	STableRow::ConstructInternal(
+		STableRow::FArguments()
+		.Padding(5.0f)
+		.ShowSelection(true)
+		.OnDragDetected(EditPanel.Get(), &SPropertyAnimatorCoreEditorEditPanel::OnDragStart),
+		InOwnerTableView
+	);
+
+	SetBorderImage(TAttribute<const FSlateBrush*>(this, &SPropertyAnimatorCoreEditorPropertiesViewControllerTableRow::GetBorder));
+}
+
+void SPropertyAnimatorCoreEditorPropertiesViewControllerTableRow::OnGlobalSelectionChanged()
+{
+	const TSharedPtr<SPropertyAnimatorCoreEditorPropertiesViewTableRow> Row = RowWeak.Pin();
+
+	if (!Row.IsValid())
+	{
+		return;
+	}
+
+	const TSharedPtr<SPropertyAnimatorCoreEditorPropertiesView> PropertiesView = Row->ViewWeak.Pin();
+
+	if (!PropertiesView.IsValid())
+	{
+		return;
+	}
+
+	const TSharedPtr<SPropertyAnimatorCoreEditorEditPanel> EditPanel = PropertiesView->GetEditPanel();
+
+	if (!EditPanel.IsValid())
+	{
+		return;
+	}
+
+	const FPropertiesViewControllerItemPtr RowItem = TileItemWeak.Pin();
+
+	if (!RowItem.IsValid())
+	{
+		return;
+	}
+
+	const TSet<FPropertiesViewControllerItem>& GlobalSelection = EditPanel->GetGlobalSelection();
+
+	if (GlobalSelection.Contains(*RowItem))
+	{
+		Row->ControllersTile->SetItemSelection(RowItem, true, ESelectInfo::Direct);
+	}
+	else
+	{
+		Row->ControllersTile->SetItemSelection(RowItem, false, ESelectInfo::Direct);
+	}
+}
diff --git a/Engine/Plugins/Experimental/PropertyAnimatorCore/Source/PropertyAnimatorCoreEditor/Private/Widgets/TableRows/SPropertyAnimatorCoreEditorPropertiesViewControllerTableRow.h b/Engine/Plugins/Experimental/PropertyAnimatorCore/Source/PropertyAnimatorCoreEditor/Private/Widgets/TableRows/SPropertyAnimatorCoreEditorPropertiesViewControllerTableRow.h
new file mode 100644
--- /dev/null
+++ b/Engine/Plugins/Experimental/PropertyAnimatorCore/Source/PropertyAnimatorCoreEditor/Private/Widgets/TableRows/SPropertyAnimatorCoreEditorPropertiesViewControllerTableRow.h
@@ -0,0 +1,27 @@
+// Copyright Epic Games, Inc. All Rights Reserved.
+
+#pragma once
+
+#include "Templates/SharedPointer.h"
+#include "Widgets/Views/SPropertyAnimatorCoreEditorPropertiesView.h"
+#include "Widgets/Views/STableRow.h"
+
+class SPropertyAnimatorCoreEditorPropertiesViewTableRow;
+
+class SPropertyAnimatorCoreEditorPropertiesViewControllerTableRow : public STableRow<FPropertiesViewControllerItemPtr>
+{
+public:
+	SLATE_BEGIN_ARGS(SPropertyAnimatorCoreEditorPropertiesViewControllerTableRow) {}
+	SLATE_END_ARGS()
+
+	virtual ~SPropertyAnimatorCoreEditorPropertiesViewControllerTableRow() override;
+
+	void Construct(const FArguments& InArgs, const TSharedRef<STableViewBase>& InOwnerTableView,
+		TSharedPtr<SPropertyAnimatorCoreEditorPropertiesViewTableRow> InView, FPropertiesViewControllerItemPtr InItem);
+
+private:
+	void OnGlobalSelectionChanged();
+
+	TWeakPtr<SPropertyAnimatorCoreEditorPropertiesViewTableRow> RowWeak;
+	FPropertiesViewControllerItemPtrWeak TileItemWeak;
+};
\ No newline at end of file
diff --git a/Engine/Plugins/Experimental/PropertyAnimatorCore/Source/PropertyAnimatorCoreEditor/Private/Widgets/TableRows/SPropertyAnimatorCoreEditorPropertiesViewTableRow.cpp b/Engine/Plugins/Experimental/PropertyAnimatorCore/Source/PropertyAnimatorCoreEditor/Private/Widgets/TableRows/SPropertyAnimatorCoreEditorPropertiesViewTableRow.cpp
new file mode 100644
--- /dev/null
+++ b/Engine/Plugins/Experimental/PropertyAnimatorCore/Source/PropertyAnimatorCoreEditor/Private/Widgets/TableRows/SPropertyAnimatorCoreEditorPropertiesViewTableRow.cpp
@@ -0,0 +1,313 @@
+// Copyright Epic Games, Inc. All Rights Reserved.
+
+#include "Widgets/TableRows/SPropertyAnimatorCoreEditorPropertiesViewTableRow.h"
+
+#include "Framework/Application/SlateApplication.h"
+#include "Menus/PropertyAnimatorCoreEditorMenu.h"
+#include "Menus/PropertyAnimatorCoreEditorMenuContext.h"
+#include "Properties/PropertyAnimatorCoreResolver.h"
+#include "Styles/PropertyAnimatorCoreEditorStyle.h"
+#include "Subsystems/PropertyAnimatorCoreEditorSubsystem.h"
+#include "Subsystems/PropertyAnimatorCoreSubsystem.h"
+#include "ToolMenu.h"
+#include "ToolMenus.h"
+#include "Widgets/DragDropOps/PropertyAnimatorCoreEditorViewDragDropOp.h"
+#include "Widgets/Input/SComboButton.h"
+#include "Widgets/SPropertyAnimatorCoreEditorEditPanel.h"
+#include "Widgets/TableRows/SPropertyAnimatorCoreEditorPropertiesViewControllerTableRow.h"
+#include "Widgets/Views/STileView.h"
+
+#define LOCTEXT_NAMESPACE "PropertyAnimatorCoreEditorPropertiesViewTableRow"
+
+void SPropertyAnimatorCoreEditorPropertiesViewTableRow::Construct(const FArguments& InArgs, const TSharedRef<STableViewBase>& InOwnerTableView, TSharedPtr<SPropertyAnimatorCoreEditorPropertiesView> InView, FPropertiesViewItemPtr InItem)
+{
+	ViewWeak = InView;
+	RowItemWeak = InItem;
+
+	check(InItem.IsValid() && InView.IsValid())
+
+	SMultiColumnTableRow<FPropertiesViewItemPtr>::Construct(
+		FSuperRowType::FArguments()
+		.Style(&FAppStyle::Get().GetWidgetStyle<FTableRowStyle>("SceneOutliner.TableViewRow"))
+		.Padding(5.0f)
+		.ShowSelection(true)
+		.OnPaintDropIndicator(this, &SPropertyAnimatorCoreEditorPropertiesViewTableRow::OnPropertyPaintDropIndicator)
+		.OnDrop(this, &SPropertyAnimatorCoreEditorPropertiesViewTableRow::OnPropertyDrop)
+		.OnCanAcceptDrop(this, &SPropertyAnimatorCoreEditorPropertiesViewTableRow::OnPropertyCanAcceptDrop),
+		InOwnerTableView
+	);
+
+	SetBorderImage(TAttribute<const FSlateBrush*>(this, &SPropertyAnimatorCoreEditorPropertiesViewTableRow::GetBorder));
+}
+
+TSharedRef<SWidget> SPropertyAnimatorCoreEditorPropertiesViewTableRow::GenerateWidgetForColumn(const FName& InColumnName)
+{
+	const FPropertiesViewItemPtr Item = RowItemWeak.Pin();
+
+	if (!Item.IsValid())
+	{
+		return SNullWidget::NullWidget;
+	}
+
+	TSharedRef<SBox> BoxWrapper = SNew(SBox)
+		.Padding(FMargin(4.0f, 0.0f))
+		.VAlign(VAlign_Center);
+
+	const bool bIsLeafRow = Item->Children.IsEmpty();
+
+	if (InColumnName == SPropertyAnimatorCoreEditorEditPanel::HeaderPropertyColumnName)
+	{
+		FString PropertyName = Item->Property.GetLeafPropertyName().ToString();
+
+		if (!Item->ParentWeak.IsValid())
+		{
+			if (const UPropertyAnimatorCoreResolver* Resolver = Item->Property.GetPropertyResolver())
+			{
+				PropertyName = Resolver->GetResolverName().ToString() + TEXT(".") + PropertyName;
+			}
+		}
+
+		const FText PropertyDisplayText = FText::FromString(PropertyName);
+
+		BoxWrapper->SetContent(
+			SNew(SHorizontalBox)
+			+ SHorizontalBox::Slot()
+			.AutoWidth()
+			[
+				SAssignNew(ExpanderArrowWidget, SExpanderArrow, SharedThis(this))
+				.Visibility(!bIsLeafRow ? EVisibility::Visible : EVisibility::Hidden)
+			]
+			+ SHorizontalBox::Slot()
+			.FillWidth(1.f)
+			.Padding(0.0f, 3.0f, 6.0f, 3.0f)
+			.VAlign(VAlign_Center)
+			[
+				SNew(STextBlock)
+				.Text(PropertyDisplayText)
+				.ToolTipText(PropertyDisplayText)
+			]
+		);
+	}
+	else if (InColumnName == SPropertyAnimatorCoreEditorEditPanel::HeaderAnimatorColumnName)
+	{
+		constexpr float ItemSize = 24.f;
+
+		BoxWrapper->SetContent(
+			SAssignNew(ControllersTile, STileView<FPropertiesViewControllerItemPtr>)
+			.Orientation(Orient_Horizontal)
+			.ItemHeight(ItemSize)
+			.ItemWidth(ItemSize)
+			.SelectionMode(ESelectionMode::Multi)
+			.ListItemsSource(&ControllersTileSource)
+			.ClearSelectionOnClick(true)
+			.OnSelectionChanged(this, &SPropertyAnimatorCoreEditorPropertiesViewTableRow::OnSelectionChanged)
+			.OnGenerateTile(this, &SPropertyAnimatorCoreEditorPropertiesViewTableRow::OnGenerateTile)
+		);
+
+		for (const TWeakObjectPtr<UPropertyAnimatorCoreBase>& ControllerWeak : Item->ControllersWeak)
+		{
+			const UPropertyAnimatorCoreBase* Controller = ControllerWeak.Get();
+
+			if (!Controller || !Controller->IsPropertyLinked(Item->Property))
+			{
+				continue;
+			}
+
+			FPropertiesViewControllerItemPtr ControllerItem = MakeShared<FPropertiesViewControllerItem>();
+			ControllerItem->Property = MakeShared<FPropertyAnimatorCoreData>(Item->Property);
+			ControllerItem->ControllerWeak = ControllerWeak;
+			ControllersTileSource.Add(ControllerItem);
+		}
+
+		ControllersTile->RebuildList();
+	}
+	else if (InColumnName == SPropertyAnimatorCoreEditorEditPanel::HeaderActionColumnName)
+	{
+		const UPropertyAnimatorCoreSubsystem* ControllerSubsystem = UPropertyAnimatorCoreSubsystem::Get();
+
+		// Only show menu if any controller supports this specific property
+		if (ControllerSubsystem->IsPropertySupported(Item->Property, false))
+		{
+			BoxWrapper->SetContent(
+				SNew(SComboButton)
+				.HAlign(HAlign_Center)
+				.VAlign(VAlign_Fill)
+				.MenuPlacement(EMenuPlacement::MenuPlacement_BelowAnchor)
+				.HasDownArrow(false)
+				.ButtonStyle(FAppStyle::Get(), "SimpleButton")
+				.OnGetMenuContent(this, &SPropertyAnimatorCoreEditorPropertiesViewTableRow::OnGetMenuContent)
+				.ButtonContent()
+				[
+					SNew(SImage)
+					.Image(FPropertyAnimatorCoreEditorStyle::Get().GetBrush("PropertyControlIcon.Add"))
+					.Visibility(EVisibility::SelfHitTestInvisible)
+				]
+			);
+		}
+	}
+
+	return BoxWrapper;
+}
+
+TOptional<EItemDropZone> SPropertyAnimatorCoreEditorPropertiesViewTableRow::OnPropertyCanAcceptDrop(const FDragDropEvent& InDragDropEvent, EItemDropZone InDropZone, FPropertiesViewItemPtr InItem)
+{
+	// Get dragged controllers and check that the property we drop them off is supported by one of them
+	if (const TSharedPtr<FPropertyAnimatorCoreEditorViewDragDropOp> DragDropOp = InDragDropEvent.GetOperationAs<FPropertyAnimatorCoreEditorViewDragDropOp>())
+	{
+		if (InItem.IsValid() && InItem->Property.IsResolved())
+		{
+			for (const FPropertiesViewControllerItem& DraggedItem : DragDropOp->GetDraggedItems())
+			{
+				if (const UPropertyAnimatorCoreBase* Controller = DraggedItem.ControllerWeak.Get())
+				{
+					if (Controller->IsPropertySupported(InItem->Property)
+						&& !Controller->IsPropertyLinked(InItem->Property))
+					{
+						return EItemDropZone::OntoItem;
+					}
+				}
+			}
+		}
+	}
+
+	return TOptional<EItemDropZone>();
+}
+
+int32 SPropertyAnimatorCoreEditorPropertiesViewTableRow::OnPropertyPaintDropIndicator(EItemDropZone InDropZone, const FPaintArgs& InPaintArgs, const FGeometry& InGeometry, const FSlateRect& InSlateRect, FSlateWindowElementList& OutElements, int32 InLayerIndex, const FWidgetStyle& InWidgetStyle, bool InbParentEnabled)
+{
+	const FSlateBrush* DropIndicatorBrush = GetDropIndicatorBrush(InDropZone);
+
+	constexpr float OffsetX = 10.0f;
+	const FVector2D Offset(OffsetX * GetIndentLevel(), 0.f);
+	FSlateDrawElement::MakeBox
+	(
+		OutElements,
+		InLayerIndex++,
+		InGeometry.ToPaintGeometry(FVector2D(InGeometry.GetLocalSize() - Offset), FSlateLayoutTransform(Offset)),
+		DropIndicatorBrush,
+		ESlateDrawEffect::None,
+		DropIndicatorBrush->GetTint(InWidgetStyle) * InWidgetStyle.GetColorAndOpacityTint()
+	);
+
+	return InLayerIndex;
+}
+
+FReply SPropertyAnimatorCoreEditorPropertiesViewTableRow::OnPropertyDrop(FDragDropEvent const& InDragDropEvent)
+{
+	if (const TSharedPtr<FPropertyAnimatorCoreEditorViewDragDropOp> DragDropOp = InDragDropEvent.GetOperationAs<FPropertyAnimatorCoreEditorViewDragDropOp>())
+	{
+		if (const FPropertiesViewItemPtr RowItem = RowItemWeak.Pin())
+		{
+			const bool bIsControlPressed = FSlateApplication::Get().GetModifierKeys().IsControlDown();
+
+			for (const FPropertiesViewControllerItem& DraggedItem : DragDropOp->GetDraggedItems())
+			{
+				if (UPropertyAnimatorCoreBase* Controller = DraggedItem.ControllerWeak.Get())
+				{
+					if (Controller->IsPropertySupported(RowItem->Property))
+					{
+						if (!bIsControlPressed && DraggedItem.Property.IsValid())
+						{
+							Controller->UnlinkProperty(*DraggedItem.Property);
+						}
+
+						Controller->LinkProperty(RowItem->Property);
+					}
+				}
+			}
+		}
+
+		return FReply::Handled().EndDragDrop();
+	}
+
+	return FReply::Unhandled();
+}
+
+void SPropertyAnimatorCoreEditorPropertiesViewTableRow::OnSelectionChanged(FPropertiesViewControllerItemPtr InItem, ESelectInfo::Type InSelectInfo)
+{
+	const TSharedPtr<SPropertyAnimatorCoreEditorPropertiesView> PropertiesView = ViewWeak.Pin();
+
+	if (!PropertiesView.IsValid() || InSelectInfo == ESelectInfo::Direct)
+	{
+		 return;
+	}
+
+	const TSharedPtr<SPropertyAnimatorCoreEditorEditPanel> EditPanel = PropertiesView->GetEditPanel();
+
+	if (!EditPanel.IsValid())
+	{
+		return;
+	}
+
+	TSet<FPropertiesViewControllerItem>& GlobalSelection = EditPanel->GetGlobalSelection();
+
+	if (!FSlateApplication::Get().GetModifierKeys().IsControlDown())
+	{
+		GlobalSelection.Reset();
+	}
+
+	for (FPropertiesViewControllerItemPtr SelectedItem : ControllersTile->GetSelectedItems())
+	{
+		GlobalSelection.Add(*SelectedItem);
+	}
+
+	EditPanel->OnGlobalSelectionChangedDelegate.Broadcast();
+}
+
+TSharedRef<ITableRow> SPropertyAnimatorCoreEditorPropertiesViewTableRow::OnGenerateTile(FPropertiesViewControllerItemPtr InItem, const TSharedRef<STableViewBase>& InOwnerTable)
+{
+	return SNew(SPropertyAnimatorCoreEditorPropertiesViewControllerTableRow, InOwnerTable, SharedThis(this), InItem);
+}
+
+TSharedRef<SWidget> SPropertyAnimatorCoreEditorPropertiesViewTableRow::OnGetMenuContent()
+{
+	static const FName AddControllersMenuName = TEXT("NewAnimatorPropertyViewMenu");
+
+	UToolMenus* ToolMenus = UToolMenus::Get();
+
+	check(ToolMenus);
+
+	if (!ToolMenus->IsMenuRegistered(AddControllersMenuName))
+	{
+		UToolMenu* AddControllersMenu = ToolMenus->RegisterMenu(AddControllersMenuName);
+		AddControllersMenu->AddDynamicSection(TEXT("FillNewAnimatorPropertyViewSection"), FNewToolMenuDelegate::CreateSP(this, &SPropertyAnimatorCoreEditorPropertiesViewTableRow::FillNewAnimatorPropertyViewSection));
+	}
+
+	UPropertyAnimatorCoreEditorMenuContext* MenuContextObject = NewObject<UPropertyAnimatorCoreEditorMenuContext>();
+
+	if (const FPropertiesViewItemPtr RowItem = RowItemWeak.Pin())
+	{
+		MenuContextObject->SetPropertyData(RowItem->Property);
+	}
+
+	const FToolMenuContext MenuContext(MenuContextObject);
+	return ToolMenus->GenerateWidget(AddControllersMenuName, MenuContext);
+}
+
+void SPropertyAnimatorCoreEditorPropertiesViewTableRow::FillNewAnimatorPropertyViewSection(UToolMenu* InToolMenu) const
+{
+	if (!InToolMenu)
+	{
+		return;
+	}
+
+	const UPropertyAnimatorCoreEditorMenuContext* Context = InToolMenu->FindContext<UPropertyAnimatorCoreEditorMenuContext>();
+
+	if (!Context)
+	{
+		return;
+	}
+
+	UPropertyAnimatorCoreEditorSubsystem* EditorSubsystem = UPropertyAnimatorCoreEditorSubsystem::Get();
+
+	if (!EditorSubsystem)
+	{
+		return;
+	}
+
+	const FPropertyAnimatorCoreEditorMenuContext MenuContext({}, {Context->GetPropertyData()});
+	const FPropertyAnimatorCoreEditorMenuOptions MenuOptions({EPropertyAnimatorCoreEditorMenuType::New});
+	EditorSubsystem->FillAnimatorMenu(InToolMenu, MenuContext, MenuOptions);
+}
+
+#undef LOCTEXT_NAMESPACE
diff --git a/Engine/Plugins/Experimental/PropertyAnimatorCore/Source/PropertyAnimatorCoreEditor/Private/Widgets/TableRows/SPropertyAnimatorCoreEditorPropertiesViewTableRow.h b/Engine/Plugins/Experimental/PropertyAnimatorCore/Source/PropertyAnimatorCoreEditor/Private/Widgets/TableRows/SPropertyAnimatorCoreEditorPropertiesViewTableRow.h
new file mode 100644
--- /dev/null
+++ b/Engine/Plugins/Experimental/PropertyAnimatorCore/Source/PropertyAnimatorCoreEditor/Private/Widgets/TableRows/SPropertyAnimatorCoreEditorPropertiesViewTableRow.h
@@ -0,0 +1,47 @@
+// Copyright Epic Games, Inc. All Rights Reserved.
+
+#pragma once
+
+#include "Templates/SharedPointer.h"
+#include "Widgets/Views/SPropertyAnimatorCoreEditorPropertiesView.h"
+#include "Widgets/Views/STableRow.h"
+#include "Widgets/Views/STileView.h"
+#include "Widgets/Views/STreeView.h"
+
+class SPropertyAnimatorCoreEditorPropertiesView;
+class UToolMenu;
+
+class SPropertyAnimatorCoreEditorPropertiesViewTableRow : public SMultiColumnTableRow<FPropertiesViewItemPtr>
+{
+	friend class SPropertyAnimatorCoreEditorPropertiesViewControllerTableRow;
+
+public:
+	SLATE_BEGIN_ARGS(SPropertyAnimatorCoreEditorPropertiesViewTableRow) {}
+	SLATE_END_ARGS()
+
+	void Construct(const FArguments& InArgs, const TSharedRef<STableViewBase>& InOwnerTableView,
+		TSharedPtr<SPropertyAnimatorCoreEditorPropertiesView> InView, FPropertiesViewItemPtr InItem);
+
+	virtual TSharedRef<SWidget> GenerateWidgetForColumn(const FName& InColumnName) override;
+
+protected:
+	TOptional<EItemDropZone> OnPropertyCanAcceptDrop(const FDragDropEvent& InDragDropEvent, EItemDropZone InDropZone, FPropertiesViewItemPtr InItem);
+
+	int32 OnPropertyPaintDropIndicator(EItemDropZone InDropZone, const FPaintArgs& InPaintArgs, const FGeometry& InGeometry, const FSlateRect& InSlateRect, FSlateWindowElementList& OutElements, int32 InLayerIndex, const FWidgetStyle& InWidgetStyle, bool InbParentEnabled);
+
+	FReply OnPropertyDrop(FDragDropEvent const& InDragDropEvent);
+
+	void OnSelectionChanged(FPropertiesViewControllerItemPtr InItem, ESelectInfo::Type InSelectInfo);
+
+	TSharedRef<ITableRow> OnGenerateTile(FPropertiesViewControllerItemPtr InItem, const TSharedRef<STableViewBase>& InOwnerTable);
+
+	TSharedRef<SWidget> OnGetMenuContent();
+
+	void FillNewAnimatorPropertyViewSection(UToolMenu* InToolMenu) const;
+
+	TSharedPtr<STileView<FPropertiesViewControllerItemPtr>> ControllersTile;
+	TArray<FPropertiesViewControllerItemPtr> ControllersTileSource;
+
+	TWeakPtr<SPropertyAnimatorCoreEditorPropertiesView> ViewWeak;
+	FPropertiesViewItemPtrWeak RowItemWeak;
+};
\ No newline at end of file
diff --git a/Engine/Plugins/Experimental/PropertyAnimatorCore/Source/PropertyAnimatorCoreEditor/Private/Widgets/Views/PropertyAnimatorCoreEditorViewItem.h b/Engine/Plugins/Experimental/PropertyAnimatorCore/Source/PropertyAnimatorCoreEditor/Private/Widgets/Views/PropertyAnimatorCoreEditorViewItem.h
new file mode 100644
--- /dev/null
+++ b/Engine/Plugins/Experimental/PropertyAnimatorCore/Source/PropertyAnimatorCoreEditor/Private/Widgets/Views/PropertyAnimatorCoreEditorViewItem.h
@@ -0,0 +1,59 @@
+// Copyright Epic Games, Inc. All Rights Reserved.
+
+#pragma once
+
+#include "Animators/PropertyAnimatorCoreBase.h"
+#include "Properties/PropertyAnimatorCoreData.h"
+
+struct FPropertiesViewControllerItem
+{
+	/** The controller containing the property */
+	TWeakObjectPtr<UPropertyAnimatorCoreBase> ControllerWeak;
+
+	/** The optional property linked to the controller */
+	TSharedPtr<FPropertyAnimatorCoreData> Property;
+
+	friend uint32 GetTypeHash(const FPropertiesViewControllerItem& InItem)
+	{
+		return HashCombine(GetTypeHash(InItem.ControllerWeak.Get()), GetTypeHash(InItem.Property.Get()));
+	}
+
+	bool operator==(const FPropertiesViewControllerItem& InOther) const
+	{
+		return ControllerWeak == InOther.ControllerWeak
+			&& Property == InOther.Property;
+	}
+};
+
+using FPropertiesViewControllerItemPtr = TSharedPtr<FPropertiesViewControllerItem>;
+using FPropertiesViewControllerItemPtrWeak = TWeakPtr<FPropertiesViewControllerItem>;
+
+struct FPropertiesViewItem
+{
+	/** The property displayed on the row */
+	FPropertyAnimatorCoreData Property;
+
+	/** Controllers this property is linked to */
+	TSet<TWeakObjectPtr<UPropertyAnimatorCoreBase>> ControllersWeak;
+
+	/** Nested properties below this property */
+	TArray<TSharedPtr<FPropertiesViewItem>> Children;
+
+	/** Parent property of this one */
+	TWeakPtr<FPropertiesViewItem> ParentWeak;
+};
+
+using FPropertiesViewItemPtr = TSharedPtr<FPropertiesViewItem>;
+using FPropertiesViewItemPtrWeak = TWeakPtr<FPropertiesViewItem>;
+
+struct FControllersViewItem
+{
+	FPropertiesViewControllerItem ControlledProperty;
+
+	TArray<TSharedPtr<FControllersViewItem>> Children;
+
+	TWeakPtr<FControllersViewItem> ParentWeak;
+};
+
+using FControllersViewItemPtr = TSharedPtr<FControllersViewItem>;
+using FControllersViewItemPtrWeak = TWeakPtr<FControllersViewItem>;
\ No newline at end of file
diff --git a/Engine/Plugins/Experimental/PropertyAnimatorCore/Source/PropertyAnimatorCoreEditor/Private/Widgets/Views/SPropertyAnimatorCoreEditorControllersView.cpp b/Engine/Plugins/Experimental/PropertyAnimatorCore/Source/PropertyAnimatorCoreEditor/Private/Widgets/Views/SPropertyAnimatorCoreEditorControllersView.cpp
new file mode 100644
--- /dev/null
+++ b/Engine/Plugins/Experimental/PropertyAnimatorCore/Source/PropertyAnimatorCoreEditor/Private/Widgets/Views/SPropertyAnimatorCoreEditorControllersView.cpp
@@ -0,0 +1,171 @@
+// Copyright Epic Games, Inc. All Rights Reserved.
+
+#include "Widgets/Views/SPropertyAnimatorCoreEditorControllersView.h"
+
+#include "Framework/Application/SlateApplication.h"
+#include "Subsystems/PropertyAnimatorCoreSubsystem.h"
+#include "Widgets/SPropertyAnimatorCoreEditorEditPanel.h"
+#include "Widgets/TableRows/SPropertyAnimatorCoreEditorControllersViewTableRow.h"
+
+#define LOCTEXT_NAMESPACE "PropertyAnimatorCoreEditorControllersView"
+
+void SPropertyAnimatorCoreEditorControllersView::Construct(const FArguments& InArgs, TSharedPtr<SPropertyAnimatorCoreEditorEditPanel> InEditPanel)
+{
+	EditPanelWeak = InEditPanel;
+
+	check(InEditPanel.IsValid());
+
+	ChildSlot
+	[
+		SNew(SOverlay)
+		+ SOverlay::Slot(0)
+		.HAlign(HAlign_Fill)
+		.VAlign(VAlign_Fill)
+		[
+			SAssignNew(ControllersTree, STreeView<FControllersViewItemPtr>)
+			.TreeItemsSource(&ControllersTreeSource)
+			.SelectionMode(ESelectionMode::Multi)
+			.ItemHeight(30.f)
+			.HeaderRow(
+				SNew(SHeaderRow)
+				+ SHeaderRow::Column(SPropertyAnimatorCoreEditorEditPanel::HeaderAnimatorColumnName)
+				.DefaultLabel(FText::FromName(SPropertyAnimatorCoreEditorEditPanel::HeaderAnimatorColumnName))
+				.FillWidth(0.5f)
+				+ SHeaderRow::Column(SPropertyAnimatorCoreEditorEditPanel::HeaderPropertyColumnName)
+				.DefaultLabel(FText::FromName(SPropertyAnimatorCoreEditorEditPanel::HeaderPropertyColumnName))
+				.FillWidth(0.5f)
+			)
+			.OnSelectionChanged(this, &SPropertyAnimatorCoreEditorControllersView::OnSelectionChanged)
+			.OnGetChildren(this, &SPropertyAnimatorCoreEditorControllersView::OnGetChildren)
+			.OnGenerateRow(this, &SPropertyAnimatorCoreEditorControllersView::OnGenerateRow)
+		]
+		+ SOverlay::Slot(1)
+		.HAlign(HAlign_Center)
+		.VAlign(VAlign_Center)
+		[
+		   SNew(STextBlock)
+		   .Text(LOCTEXT("ControllersListEmpty", "No animators found on this actor"))
+		   .Visibility(this, &SPropertyAnimatorCoreEditorControllersView::GetControllerTextVisibility)
+		]
+	];
+}
+
+void SPropertyAnimatorCoreEditorControllersView::Update()
+{
+	const TSharedPtr<SPropertyAnimatorCoreEditorEditPanel> EditPanel = EditPanelWeak.Pin();
+
+	if (!EditPanel.IsValid())
+	{
+		return;
+	}
+
+	const AActor* ContextActor = EditPanel->GetOptions().GetContextActor();
+
+	if (!ContextActor)
+	{
+		return;
+	}
+
+	// Collect all existing controllers on actor
+	const UPropertyAnimatorCoreSubsystem* PropertyControllerSubsystem = UPropertyAnimatorCoreSubsystem::Get();
+
+	if (!PropertyControllerSubsystem)
+	{
+		return;
+	}
+
+	TSet<UPropertyAnimatorCoreBase*> ExistingControllers = PropertyControllerSubsystem->GetExistingAnimators(ContextActor);
+	ControllersTreeSource.Empty(ExistingControllers.Num());
+
+	for (UPropertyAnimatorCoreBase* ExistingController : ExistingControllers)
+	{
+		FControllersViewItemPtr NewItem = MakeShared<FControllersViewItem>();
+		NewItem->ControlledProperty.ControllerWeak = ExistingController;
+		ControllersTreeSource.Add(NewItem);
+	}
+
+	if (ControllersTree.IsValid())
+	{
+		ControllersTree->RebuildList();
+
+		// Auto expand
+		for (const FControllersViewItemPtr& Item : ControllersTreeSource)
+		{
+			ControllersTree->SetItemExpansion(Item, true);
+		}
+	}
+}
+
+void SPropertyAnimatorCoreEditorControllersView::OnSelectionChanged(FControllersViewItemPtr InItem, ESelectInfo::Type InSelectInfo)
+{
+	if (InSelectInfo == ESelectInfo::Direct)
+	{
+		return;
+	}
+
+	const TSharedPtr<SPropertyAnimatorCoreEditorEditPanel> EditPanel = EditPanelWeak.Pin();
+
+	if (!EditPanel.IsValid())
+	{
+		return;
+	}
+
+	TSet<FPropertiesViewControllerItem>& GlobalSelection = EditPanel->GetGlobalSelection();
+
+	if (!FSlateApplication::Get().GetModifierKeys().IsControlDown())
+	{
+		GlobalSelection.Reset();
+	}
+
+	for (const FControllersViewItemPtr& SelectedItem : ControllersTree->GetSelectedItems())
+	{
+		GlobalSelection.Add(SelectedItem->ControlledProperty);
+	}
+
+	EditPanel->OnGlobalSelectionChangedDelegate.Broadcast();
+}
+
+void SPropertyAnimatorCoreEditorControllersView::OnGetChildren(FControllersViewItemPtr InItem, TArray<FControllersViewItemPtr>& OutChildren)
+{
+	if (!InItem.IsValid())
+	{
+		 return;
+	}
+
+	// Children already loaded
+	if (!InItem->Children.IsEmpty())
+	{
+		OutChildren = InItem->Children;
+		return;
+	}
+
+	UPropertyAnimatorCoreBase* Controller = InItem->ControlledProperty.ControllerWeak.Get();
+
+	// We reach leaf when we have a property set
+	if (!Controller
+		|| InItem->ControlledProperty.Property.IsValid())
+	{
+		return;
+	}
+
+	for (const FPropertyAnimatorCoreData& LinkedProperty : Controller->GetLinkedProperties())
+	{
+		FControllersViewItemPtr NewItem = MakeShared<FControllersViewItem>();
+		NewItem->ControlledProperty.ControllerWeak = Controller;
+		NewItem->ControlledProperty.Property = MakeShared<FPropertyAnimatorCoreData>(LinkedProperty);
+		NewItem->ParentWeak = InItem;
+		OutChildren.Add(NewItem);
+	}
+}
+
+TSharedRef<ITableRow> SPropertyAnimatorCoreEditorControllersView::OnGenerateRow(FControllersViewItemPtr InItem, const TSharedRef<STableViewBase>& InOwnerTable)
+{
+	return SNew(SPropertyAnimatorCoreEditorControllersViewTableRow, InOwnerTable, SharedThis(this), InItem);
+}
+
+EVisibility SPropertyAnimatorCoreEditorControllersView::GetControllerTextVisibility() const
+{
+	return ControllersTreeSource.IsEmpty() ? EVisibility::Visible : EVisibility::Collapsed;
+}
+
+#undef LOCTEXT_NAMESPACE
\ No newline at end of file
diff --git a/Engine/Plugins/Experimental/PropertyAnimatorCore/Source/PropertyAnimatorCoreEditor/Private/Widgets/Views/SPropertyAnimatorCoreEditorControllersView.h b/Engine/Plugins/Experimental/PropertyAnimatorCore/Source/PropertyAnimatorCoreEditor/Private/Widgets/Views/SPropertyAnimatorCoreEditorControllersView.h
new file mode 100644
--- /dev/null
+++ b/Engine/Plugins/Experimental/PropertyAnimatorCore/Source/PropertyAnimatorCoreEditor/Private/Widgets/Views/SPropertyAnimatorCoreEditorControllersView.h
@@ -0,0 +1,42 @@
+// Copyright Epic Games, Inc. All Rights Reserved.
+
+#pragma once
+
+#include "Templates/SharedPointer.h"
+#include "Widgets/Views/PropertyAnimatorCoreEditorViewItem.h"
+#include "Widgets/Views/STreeView.h"
+#include "Widgets/Views/SHeaderRow.h"
+#include "Widgets/SCompoundWidget.h"
+
+class SPropertyAnimatorCoreEditorEditPanel;
+
+class SPropertyAnimatorCoreEditorControllersView : public SCompoundWidget
+{
+	friend class SPropertyAnimatorCoreEditorControllersViewTableRow;
+
+public:
+	SLATE_BEGIN_ARGS(SPropertyAnimatorCoreEditorControllersView) {}
+	SLATE_END_ARGS()
+
+	void Construct(const FArguments& InArgs, TSharedPtr<SPropertyAnimatorCoreEditorEditPanel> InEditPanel);
+
+	void Update();
+
+	TSharedPtr<SPropertyAnimatorCoreEditorEditPanel> GetEditPanel() const
+	{
+		return EditPanelWeak.Pin();
+	}
+
+protected:
+	void OnSelectionChanged(FControllersViewItemPtr InItem, ESelectInfo::Type InSelectInfo);
+
+	void OnGetChildren(FControllersViewItemPtr InItem, TArray<FControllersViewItemPtr>& OutChildren);
+
+	TSharedRef<ITableRow> OnGenerateRow(FControllersViewItemPtr ControllersViewItem, const TSharedRef<STableViewBase>& InOwnerTable);
+
+	EVisibility GetControllerTextVisibility() const;
+
+	TSharedPtr<STreeView<FControllersViewItemPtr>> ControllersTree;
+	TArray<FControllersViewItemPtr> ControllersTreeSource;
+	TWeakPtr<SPropertyAnimatorCoreEditorEditPanel> EditPanelWeak;
+};
diff --git a/Engine/Plugins/Experimental/PropertyAnimatorCore/Source/PropertyAnimatorCoreEditor/Private/Widgets/Views/SPropertyAnimatorCoreEditorPropertiesView.cpp b/Engine/Plugins/Experimental/PropertyAnimatorCore/Source/PropertyAnimatorCoreEditor/Private/Widgets/Views/SPropertyAnimatorCoreEditorPropertiesView.cpp
new file mode 100644
--- /dev/null
+++ b/Engine/Plugins/Experimental/PropertyAnimatorCore/Source/PropertyAnimatorCoreEditor/Private/Widgets/Views/SPropertyAnimatorCoreEditorPropertiesView.cpp
@@ -0,0 +1,256 @@
+// Copyright Epic Games, Inc. All Rights Reserved.
+
+#include "Widgets/Views/SPropertyAnimatorCoreEditorPropertiesView.h"
+
+#include "Animators/PropertyAnimatorCoreBase.h"
+#include "Subsystems/PropertyAnimatorCoreSubsystem.h"
+#include "Widgets/SPropertyAnimatorCoreEditorEditPanel.h"
+#include "Widgets/TableRows/SPropertyAnimatorCoreEditorPropertiesViewTableRow.h"
+
+#define LOCTEXT_NAMESPACE "PropertyAnimatorCoreEditorPropertiesView"
+
+void SPropertyAnimatorCoreEditorPropertiesView::Construct(const FArguments& InArgs, TSharedPtr<SPropertyAnimatorCoreEditorEditPanel> InEditPanel)
+{
+	EditPanelWeak = InEditPanel;
+
+	check(InEditPanel.IsValid());
+
+	ChildSlot
+	[
+		SNew(SOverlay)
+		+ SOverlay::Slot(0)
+		.HAlign(HAlign_Fill)
+		.VAlign(VAlign_Fill)
+		[
+			SAssignNew(PropertiesTree, STreeView<FPropertiesViewItemPtr>)
+		   .TreeItemsSource(&PropertiesTreeSource)
+		   .SelectionMode(ESelectionMode::None)
+		   .ItemHeight(30.f)
+		   .HeaderRow(
+			   SNew(SHeaderRow)
+			   + SHeaderRow::Column(SPropertyAnimatorCoreEditorEditPanel::HeaderPropertyColumnName)
+			   .DefaultLabel(FText::FromName(SPropertyAnimatorCoreEditorEditPanel::HeaderPropertyColumnName))
+			   .FillWidth(0.4f)
+			   + SHeaderRow::Column(SPropertyAnimatorCoreEditorEditPanel::HeaderAnimatorColumnName)
+			   .DefaultLabel(FText::FromName(SPropertyAnimatorCoreEditorEditPanel::HeaderAnimatorColumnName))
+			   .FillWidth(0.5f)
+			   + SHeaderRow::Column(SPropertyAnimatorCoreEditorEditPanel::HeaderActionColumnName)
+			   .DefaultLabel(FText::GetEmpty())
+			   .FillWidth(0.1f)
+		   )
+		   .OnGetChildren(this, &SPropertyAnimatorCoreEditorPropertiesView::OnGetChildren)
+		   .OnGenerateRow(this, &SPropertyAnimatorCoreEditorPropertiesView::OnGenerateRow)
+		]
+		+ SOverlay::Slot(1)
+		.HAlign(HAlign_Center)
+		.VAlign(VAlign_Center)
+		[
+			SNew(STextBlock)
+	        .Text(LOCTEXT("PropertiesListEmpty", "No linked properties found on this actor"))
+	        .Visibility(this, &SPropertyAnimatorCoreEditorPropertiesView::GetPropertyTextVisibility)
+		]
+	];
+}
+
+void SPropertyAnimatorCoreEditorPropertiesView::Update()
+{
+	const TSharedPtr<SPropertyAnimatorCoreEditorEditPanel> EditPanel = EditPanelWeak.Pin();
+
+	if (!EditPanel.IsValid())
+	{
+		return;
+	}
+
+	AActor* ContextActor = EditPanel->GetOptions().GetContextActor();
+
+	if (!ContextActor)
+	{
+		return;
+	}
+
+	// Collect all existing controllers on actor
+	const UPropertyAnimatorCoreSubsystem* PropertyControllerSubsystem = UPropertyAnimatorCoreSubsystem::Get();
+
+	if (!PropertyControllerSubsystem)
+	{
+		return;
+	}
+
+	TSet<UPropertyAnimatorCoreBase*> ExistingControllers = PropertyControllerSubsystem->GetExistingAnimators(ContextActor);
+
+	// Collect all linked root properties and their controllers
+	TMap<FPropertyAnimatorCoreData, TSet<TWeakObjectPtr<UPropertyAnimatorCoreBase>>> LinkedPropertiesMap;
+	for (UPropertyAnimatorCoreBase* Controller : ExistingControllers)
+	{
+		for (const FPropertyAnimatorCoreData& LinkedProperty : Controller->GetLinkedProperties())
+		{
+			TSet<TWeakObjectPtr<UPropertyAnimatorCoreBase>>& Controllers = LinkedPropertiesMap.FindOrAdd(LinkedProperty);
+			Controllers.Add(Controller);
+		}
+	}
+
+	// Only get root member properties to build the tree
+	TMap<FPropertyAnimatorCoreData, TSet<TWeakObjectPtr<UPropertyAnimatorCoreBase>>> MemberPropertiesMap;
+	for (const TPair<FPropertyAnimatorCoreData, TSet<TWeakObjectPtr<UPropertyAnimatorCoreBase>>>& LinkedPropertyPair : LinkedPropertiesMap)
+	{
+		TOptional<FPropertyAnimatorCoreData> MemberProperty = LinkedPropertyPair.Key.GetRootParent();
+
+		if (MemberProperty.IsSet())
+		{
+			TSet<TWeakObjectPtr<UPropertyAnimatorCoreBase>>& Controllers = MemberPropertiesMap.FindOrAdd(MemberProperty.GetValue());
+			Controllers.Append(LinkedPropertyPair.Value);
+		}
+	}
+
+	// Add resolvable properties
+	FPropertyAnimatorCoreData Context(ContextActor, nullptr, nullptr);
+	TSet<FPropertyAnimatorCoreData> ResolvableProperties;
+	PropertyControllerSubsystem->GetResolvableProperties(Context, ResolvableProperties);
+
+	for (const FPropertyAnimatorCoreData& ResolvableProperty : ResolvableProperties)
+	{
+		TSet<UPropertyAnimatorCoreBase*> LinkedControllers = PropertyControllerSubsystem->GetPropertyLinkedAnimators(ResolvableProperty);
+		if (!LinkedControllers.IsEmpty())
+		{
+			TSet<TWeakObjectPtr<UPropertyAnimatorCoreBase>>& Controllers = MemberPropertiesMap.Add(ResolvableProperty);
+			Algo::Transform(LinkedControllers, Controllers, [](UPropertyAnimatorCoreBase* InController)
+			{
+				return InController;
+			});
+		}
+	}
+
+	PropertiesTreeSource.Empty(MemberPropertiesMap.Num());
+	for (const TPair<FPropertyAnimatorCoreData, TSet<TWeakObjectPtr<UPropertyAnimatorCoreBase>>>& MemberPropertyPair : MemberPropertiesMap)
+	{
+		TSharedPtr<FPropertiesViewItem> NewItem = MakeShared<FPropertiesViewItem>();
+		NewItem->Property = MemberPropertyPair.Key;
+		NewItem->ControllersWeak = MemberPropertyPair.Value;
+		PropertiesTreeSource.Add(NewItem);
+	}
+
+	if (PropertiesTree.IsValid())
+	{
+		PropertiesTree->RebuildList();
+
+		// Auto expand
+		for (const FPropertiesViewItemPtr& Item : PropertiesTreeSource)
+		{
+			PropertiesTree->SetItemExpansion(Item, true);
+		}
+	}
+}
+
+TSharedRef<ITableRow> SPropertyAnimatorCoreEditorPropertiesView::OnGenerateRow(FPropertiesViewItemPtr InItem, const TSharedRef<STableViewBase>& InOwnerTable)
+{
+	return SNew(SPropertyAnimatorCoreEditorPropertiesViewTableRow, InOwnerTable, SharedThis(this), InItem);
+}
+
+void SPropertyAnimatorCoreEditorPropertiesView::OnGetChildren(FPropertiesViewItemPtr InItem, TArray<FPropertiesViewItemPtr>& OutChildren) const
+{
+	if (!InItem.IsValid())
+	{
+		return;
+	}
+
+	// Children already loaded
+	if (!InItem->Children.IsEmpty())
+	{
+		OutChildren = InItem->Children;
+		return;
+	}
+
+	// Map properties and controllers linked
+	TMap<FPropertyAnimatorCoreData, TSet<TWeakObjectPtr<UPropertyAnimatorCoreBase>>> ChildrenPropertiesMap;
+
+	// Show resolvable properties tree
+	if (InItem->Property.IsResolvable())
+	{
+		if (UPropertyAnimatorCoreSubsystem* ControllerSubsystem = UPropertyAnimatorCoreSubsystem::Get())
+		{
+			TSet<UPropertyAnimatorCoreBase*> ExistingControllers = ControllerSubsystem->GetExistingAnimators(InItem->Property);
+			TSet<UPropertyAnimatorCoreBase*> AvailableControllers = ControllerSubsystem->GetAvailableAnimators(&InItem->Property);
+
+			constexpr bool bRecursiveSearch = true;
+
+			for (UPropertyAnimatorCoreBase* Controller : ExistingControllers)
+			{
+				TSet<FPropertyAnimatorCoreData> SupportedProperties;
+				Controller->GetPropertiesSupported(InItem->Property, SupportedProperties, bRecursiveSearch);
+
+				// Add available properties to control with this controller
+				for (FPropertyAnimatorCoreData& SupportedProperty : SupportedProperties)
+				{
+					// Get only direct child of PropertyData and not nested ones
+					if (TOptional<FPropertyAnimatorCoreData> ChildProperty = SupportedProperty.GetChildOf(InItem->Property))
+					{
+						TSet<TWeakObjectPtr<UPropertyAnimatorCoreBase>>& Controllers = ChildrenPropertiesMap.FindOrAdd(ChildProperty.GetValue());
+
+						// Check if we are controlling that property
+						if (Controller->IsPropertyLinked(ChildProperty.GetValue()))
+						{
+							Controllers.Add(Controller);
+						}
+					}
+				}
+			}
+
+			for (UPropertyAnimatorCoreBase* Controller : AvailableControllers)
+			{
+				TSet<FPropertyAnimatorCoreData> SupportedProperties;
+				Controller->GetPropertiesSupported(InItem->Property, SupportedProperties, bRecursiveSearch);
+
+				// Add all supported properties from new controller
+				for (FPropertyAnimatorCoreData& SupportedProperty : SupportedProperties)
+				{
+					// Get only direct child of PropertyData and not nested ones
+					if (TOptional<FPropertyAnimatorCoreData> ChildProperty = SupportedProperty.GetChildOf(InItem->Property))
+					{
+						ChildrenPropertiesMap.FindOrAdd(ChildProperty.GetValue());
+					}
+				}
+			}
+		}
+	}
+	else
+	{
+		for (const TWeakObjectPtr<UPropertyAnimatorCoreBase>& ControllerWeak : InItem->ControllersWeak)
+		{
+			UPropertyAnimatorCoreBase* Controller = ControllerWeak.Get();
+
+			if (!Controller)
+			{
+				continue;
+			}
+
+			for (FPropertyAnimatorCoreData& InnerProperty : Controller->GetInnerPropertiesLinked(InItem->Property))
+			{
+				TOptional<FPropertyAnimatorCoreData> DirectPropertyChild = InnerProperty.GetChildOf(InItem->Property);
+
+				if (DirectPropertyChild.IsSet())
+				{
+					TSet<TWeakObjectPtr<UPropertyAnimatorCoreBase>>& Controllers = ChildrenPropertiesMap.FindOrAdd(DirectPropertyChild.GetValue());
+					Controllers.Add(Controller);
+				}
+			}
+		}
+	}
+
+	for (const TPair<FPropertyAnimatorCoreData, TSet<TWeakObjectPtr<UPropertyAnimatorCoreBase>>>& ChildrenPropertyPair : ChildrenPropertiesMap)
+	{
+		TSharedPtr<FPropertiesViewItem> NewItem = MakeShared<FPropertiesViewItem>();
+		NewItem->Property = ChildrenPropertyPair.Key;
+		NewItem->ControllersWeak = ChildrenPropertyPair.Value;
+		NewItem->ParentWeak = InItem;
+		InItem->Children.Add(NewItem);
+	}
+
+	OutChildren = InItem->Children;
+}
+
+EVisibility SPropertyAnimatorCoreEditorPropertiesView::GetPropertyTextVisibility() const
+{
+	return PropertiesTreeSource.IsEmpty() ? EVisibility::Visible : EVisibility::Collapsed;
+}
+
+#undef LOCTEXT_NAMESPACE
diff --git a/Engine/Plugins/Experimental/PropertyAnimatorCore/Source/PropertyAnimatorCoreEditor/Private/Widgets/Views/SPropertyAnimatorCoreEditorPropertiesView.h b/Engine/Plugins/Experimental/PropertyAnimatorCore/Source/PropertyAnimatorCoreEditor/Private/Widgets/Views/SPropertyAnimatorCoreEditorPropertiesView.h
new file mode 100644
--- /dev/null
+++ b/Engine/Plugins/Experimental/PropertyAnimatorCore/Source/PropertyAnimatorCoreEditor/Private/Widgets/Views/SPropertyAnimatorCoreEditorPropertiesView.h
@@ -0,0 +1,39 @@
+// Copyright Epic Games, Inc. All Rights Reserved.
+
+#pragma once
+
+#include "Templates/SharedPointer.h"
+#include "Widgets/Views/PropertyAnimatorCoreEditorViewItem.h"
+#include "Widgets/Views/STreeView.h"
+#include "Widgets/Views/SHeaderRow.h"
+#include "Widgets/SCompoundWidget.h"
+
+class SPropertyAnimatorCoreEditorEditPanel;
+
+class SPropertyAnimatorCoreEditorPropertiesView : public SCompoundWidget
+{
+public:
+	SLATE_BEGIN_ARGS(SPropertyAnimatorCoreEditorPropertiesView) {}
+	SLATE_END_ARGS()
+
+	void Construct(const FArguments& InArgs, TSharedPtr<SPropertyAnimatorCoreEditorEditPanel> InEditPanel);
+
+	void Update();
+
+	TSharedPtr<SPropertyAnimatorCoreEditorEditPanel> GetEditPanel() const
+	{
+		return EditPanelWeak.Pin();
+	}
+
+protected:
+	TSharedRef<ITableRow> OnGenerateRow(FPropertiesViewItemPtr InItem, const TSharedRef<STableViewBase>& InOwnerTable);
+
+	void OnGetChildren(FPropertiesViewItemPtr InItem, TArray<FPropertiesViewItemPtr>& OutChildren) const;
+
+	EVisibility GetPropertyTextVisibility() const;
+
+private:
+	TSharedPtr<STreeView<FPropertiesViewItemPtr>> PropertiesTree;
+	TArray<FPropertiesViewItemPtr> PropertiesTreeSource;
+	TWeakPtr<SPropertyAnimatorCoreEditorEditPanel> EditPanelWeak;
+};
\ No newline at end of file
diff --git a/Engine/Plugins/Experimental/PropertyAnimatorCore/Source/PropertyAnimatorCoreEditor/Public/Widgets/PropertyAnimatorCoreEditorEditPanelOptions.h b/Engine/Plugins/Experimental/PropertyAnimatorCore/Source/PropertyAnimatorCoreEditor/Public/Widgets/PropertyAnimatorCoreEditorEditPanelOptions.h
new file mode 100644
--- /dev/null
+++ b/Engine/Plugins/Experimental/PropertyAnimatorCore/Source/PropertyAnimatorCoreEditor/Public/Widgets/PropertyAnimatorCoreEditorEditPanelOptions.h
@@ -0,0 +1,29 @@
+// Copyright Epic Games, Inc. All Rights Reserved.
+
+#pragma once
+
+#include "Templates/SharedPointer.h"
+#include "UObject/WeakObjectPtr.h"
+
+class AActor;
+class SPropertyAnimatorCoreEditorEditPanel;
+
+struct PROPERTYANIMATORCOREEDITOR_API FPropertyAnimatorCoreEditorEditPanelOptions
+{
+	FPropertyAnimatorCoreEditorEditPanelOptions()
+	{}
+
+	explicit FPropertyAnimatorCoreEditorEditPanelOptions(TSharedPtr<SPropertyAnimatorCoreEditorEditPanel> InPropertyPanel)
+		: EditPanelWeak(InPropertyPanel)
+	{}
+
+	void SetContextActor(AActor* InActor);
+	AActor* GetContextActor() const
+	{
+		return ContextActorWeak.Get();
+	}
+
+protected:
+	TWeakPtr<SPropertyAnimatorCoreEditorEditPanel> EditPanelWeak;
+	TWeakObjectPtr<AActor> ContextActorWeak;
+};
\ No newline at end of file
