MotionDesign : PropertyAnimator
- Use SP instead of lambda for details customization deferred action

#changelist validated
#virtualized

diff --git a/Engine/Plugins/Experimental/PropertyAnimatorCore/Source/PropertyAnimatorCoreEditor/Private/Customizations/PropertyAnimatorCoreEditorContextTypeCustomization.cpp b/Engine/Plugins/Experimental/PropertyAnimatorCore/Source/PropertyAnimatorCoreEditor/Private/Customizations/PropertyAnimatorCoreEditorContextTypeCustomization.cpp
--- a/Engine/Plugins/Experimental/PropertyAnimatorCore/Source/PropertyAnimatorCoreEditor/Private/Customizations/PropertyAnimatorCoreEditorContextTypeCustomization.cpp
+++ b/Engine/Plugins/Experimental/PropertyAnimatorCore/Source/PropertyAnimatorCoreEditor/Private/Customizations/PropertyAnimatorCoreEditorContextTypeCustomization.cpp
@@ -115,13 +115,7 @@ void FPropertyAnimatorCoreEditorContextTypeCustomization::CustomizeChildren(TSha
 						}));
 
 						// Hide customization dropdown to avoid user changing instanced struct type
-						PropertyUtilities->EnqueueDeferredAction(FSimpleDelegate::CreateLambda([Row]()
-						{
-							if (Row && Row->CustomValueWidget())
-							{
-								Row->CustomValueWidget()->Widget->SetVisibility(EVisibility::Hidden);
-							}
-						}));
+						PropertyUtilities->EnqueueDeferredAction(FSimpleDelegate::CreateSP(this, &FPropertyAnimatorCoreEditorContextTypeCustomization::HideCustomValueWidget, Row));
 					}
 				}
 			}
@@ -164,6 +158,14 @@ void FPropertyAnimatorCoreEditorContextTypeCustomization::OnPropertyEnabled(EChe
 	}
 }
 
+void FPropertyAnimatorCoreEditorContextTypeCustomization::HideCustomValueWidget(IDetailPropertyRow* InConverterRow)
+{
+	if (InConverterRow && InConverterRow->CustomValueWidget())
+	{
+		InConverterRow->CustomValueWidget()->Widget->SetVisibility(EVisibility::Hidden);
+	}
+}
+
 FReply FPropertyAnimatorCoreEditorContextTypeCustomization::UnlinkProperty() const
 {
 	if (!PropertyContextHandle.IsValid())
diff --git a/Engine/Plugins/Experimental/PropertyAnimatorCore/Source/PropertyAnimatorCoreEditor/Private/Customizations/PropertyAnimatorCoreEditorContextTypeCustomization.h b/Engine/Plugins/Experimental/PropertyAnimatorCore/Source/PropertyAnimatorCoreEditor/Private/Customizations/PropertyAnimatorCoreEditorContextTypeCustomization.h
--- a/Engine/Plugins/Experimental/PropertyAnimatorCore/Source/PropertyAnimatorCoreEditor/Private/Customizations/PropertyAnimatorCoreEditorContextTypeCustomization.h
+++ b/Engine/Plugins/Experimental/PropertyAnimatorCore/Source/PropertyAnimatorCoreEditor/Private/Customizations/PropertyAnimatorCoreEditorContextTypeCustomization.h
@@ -31,6 +31,9 @@ protected:
 	ECheckBoxState IsPropertyEnabled() const;
 	void OnPropertyEnabled(ECheckBoxState InNewState) const;
 
+	/** Hide converter custom value widget */
+	void HideCustomValueWidget(IDetailPropertyRow* InConverterRow);
+
 	FReply UnlinkProperty() const;
 
 	TSharedPtr<IPropertyHandle> PropertyContextHandle;
