MotionDesign : PropertyAnimator
- Cleaning public facing API

#changelist validated
#virtualized

diff --git a/Engine/Plugins/Experimental/PropertyAnimator/Source/PropertyAnimator/Private/Animators/PropertyAnimatorClock.cpp b/Engine/Plugins/Experimental/PropertyAnimator/Source/PropertyAnimator/Private/Animators/PropertyAnimatorClock.cpp
--- a/Engine/Plugins/Experimental/PropertyAnimator/Source/PropertyAnimator/Private/Animators/PropertyAnimatorClock.cpp
+++ b/Engine/Plugins/Experimental/PropertyAnimator/Source/PropertyAnimator/Private/Animators/PropertyAnimatorClock.cpp
@@ -34,14 +34,14 @@ void UPropertyAnimatorClock::SetDisplayFormat(const FString& InDisplayFormat)
 	DisplayFormat = InDisplayFormat;
 }
 
-void UPropertyAnimatorClock::SetCountdownDuration(const FString& InDisplayFormat)
+void UPropertyAnimatorClock::SetCountdownDuration(const FString& InDuration)
 {
-	if (DisplayFormat == InDisplayFormat)
+	if (CountdownDuration == InDuration)
 	{
 		return;
 	}
 
-	DisplayFormat = InDisplayFormat;
+	CountdownDuration = InDuration;
 	OnModeChanged();
 }
 
diff --git a/Engine/Plugins/Experimental/PropertyAnimator/Source/PropertyAnimator/Private/Animators/PropertyAnimatorSoundWave.cpp b/Engine/Plugins/Experimental/PropertyAnimator/Source/PropertyAnimator/Private/Animators/PropertyAnimatorSoundWave.cpp
--- a/Engine/Plugins/Experimental/PropertyAnimator/Source/PropertyAnimator/Private/Animators/PropertyAnimatorSoundWave.cpp
+++ b/Engine/Plugins/Experimental/PropertyAnimator/Source/PropertyAnimator/Private/Animators/PropertyAnimatorSoundWave.cpp
@@ -12,6 +12,17 @@ UPropertyAnimatorSoundWave::UPropertyAnimatorSoundWave()
 	SetAnimatorDisplayName(DefaultControllerName);
 }
 
+void UPropertyAnimatorSoundWave::SetSampledSoundWave(USoundWave* InSoundWave)
+{
+	if (SampledSoundWave == InSoundWave)
+	{
+		return;
+	}
+
+	SampledSoundWave = InSoundWave;
+	OnSampledSoundWaveChanged();
+}
+
 void UPropertyAnimatorSoundWave::SetLoop(bool bInLoop)
 {
 	bLoop = bInLoop;
diff --git a/Engine/Plugins/Experimental/PropertyAnimator/Source/PropertyAnimator/Private/Properties/PropertyAnimatorTextGroup.cpp b/Engine/Plugins/Experimental/PropertyAnimator/Source/PropertyAnimator/Private/Properties/PropertyAnimatorTextGroup.cpp
--- a/Engine/Plugins/Experimental/PropertyAnimator/Source/PropertyAnimator/Private/Properties/PropertyAnimatorTextGroup.cpp
+++ b/Engine/Plugins/Experimental/PropertyAnimator/Source/PropertyAnimator/Private/Properties/PropertyAnimatorTextGroup.cpp
@@ -6,6 +6,21 @@
 #include "Properties/PropertyAnimatorCoreResolver.h"
 #include "Text3DComponent.h"
 
+void UPropertyAnimatorTextGroup::SetRangeStart(float InRangeStart)
+{
+	RangeStart = FMath::Clamp(InRangeStart, 0, 1);
+}
+
+void UPropertyAnimatorTextGroup::SetRangeEnd(float InRangeEnd)
+{
+	RangeEnd = FMath::Clamp(InRangeEnd, 0, 1);
+}
+
+void UPropertyAnimatorTextGroup::SetRangeOffset(float InRangeOffset)
+{
+	RangeOffset = InRangeOffset;
+}
+
 void UPropertyAnimatorTextGroup::ManageProperties(const UPropertyAnimatorCoreContext* InContext, TArray<FPropertyAnimatorCoreData>& InOutProperties)
 {
 	if (InOutProperties.IsEmpty())
diff --git a/Engine/Plugins/Experimental/PropertyAnimator/Source/PropertyAnimator/Public/Animators/PropertyAnimatorBounce.h b/Engine/Plugins/Experimental/PropertyAnimator/Source/PropertyAnimator/Public/Animators/PropertyAnimatorBounce.h
--- a/Engine/Plugins/Experimental/PropertyAnimator/Source/PropertyAnimator/Public/Animators/PropertyAnimatorBounce.h
+++ b/Engine/Plugins/Experimental/PropertyAnimator/Source/PropertyAnimator/Public/Animators/PropertyAnimatorBounce.h
@@ -8,8 +8,8 @@
 /**
  * Applies an additive bounce movement with various options on supported float properties
  */
-UCLASS(AutoExpandCategories=("Animator"))
-class PROPERTYANIMATOR_API UPropertyAnimatorBounce : public UPropertyAnimatorFloatBase
+UCLASS(MinimalAPI, AutoExpandCategories=("Animator"))
+class UPropertyAnimatorBounce : public UPropertyAnimatorFloatBase
 {
 	GENERATED_BODY()
 
@@ -18,7 +18,7 @@ public:
 
 	UPropertyAnimatorBounce();
 
-	void SetInvertEffect(bool bInvert);
+	PROPERTYANIMATOR_API void SetInvertEffect(bool bInvert);
 	bool GetInvertEffect() const
 	{
 		return bInvertEffect;
diff --git a/Engine/Plugins/Experimental/PropertyAnimator/Source/PropertyAnimator/Public/Animators/PropertyAnimatorClock.h b/Engine/Plugins/Experimental/PropertyAnimator/Source/PropertyAnimator/Public/Animators/PropertyAnimatorClock.h
--- a/Engine/Plugins/Experimental/PropertyAnimator/Source/PropertyAnimator/Public/Animators/PropertyAnimatorClock.h
+++ b/Engine/Plugins/Experimental/PropertyAnimator/Source/PropertyAnimator/Public/Animators/PropertyAnimatorClock.h
@@ -19,8 +19,8 @@ enum class EPropertyAnimatorClockMode : uint8
 };
 
 /** Animate supported string properties with a clock feature */
-UCLASS(AutoExpandCategories=("Animator"))
-class PROPERTYANIMATOR_API UPropertyAnimatorClock : public UPropertyAnimatorCoreBase
+UCLASS(MinimalAPI, AutoExpandCategories=("Animator"))
+class UPropertyAnimatorClock : public UPropertyAnimatorCoreBase
 {
 	GENERATED_BODY()
 
@@ -29,19 +29,19 @@ public:
 
 	UPropertyAnimatorClock();
 
-	void SetMode(EPropertyAnimatorClockMode InMode);
+	PROPERTYANIMATOR_API void SetMode(EPropertyAnimatorClockMode InMode);
 	EPropertyAnimatorClockMode GetMode() const
 	{
 		return Mode;
 	}
 
-	void SetDisplayFormat(const FString& InDisplayFormat);
+	PROPERTYANIMATOR_API void SetDisplayFormat(const FString& InDisplayFormat);
 	const FString& GetDisplayFormat() const
 	{
 		return DisplayFormat;
 	}
 
-	void SetCountdownDuration(const FString& InDisplayFormat);
+	PROPERTYANIMATOR_API void SetCountdownDuration(const FString& InDuration);
 	const FString& GetCountdownDuration() const
 	{
 		return CountdownDuration;
diff --git a/Engine/Plugins/Experimental/PropertyAnimator/Source/PropertyAnimator/Public/Animators/PropertyAnimatorFloatBase.h b/Engine/Plugins/Experimental/PropertyAnimator/Source/PropertyAnimator/Public/Animators/PropertyAnimatorFloatBase.h
--- a/Engine/Plugins/Experimental/PropertyAnimator/Source/PropertyAnimator/Public/Animators/PropertyAnimatorFloatBase.h
+++ b/Engine/Plugins/Experimental/PropertyAnimator/Source/PropertyAnimator/Public/Animators/PropertyAnimatorFloatBase.h
@@ -10,39 +10,39 @@ class UPropertyAnimatorFloatContext;
 /**
  * Animate supported float properties with various options
  */
-UCLASS(Abstract, AutoExpandCategories=("Animator"))
-class PROPERTYANIMATOR_API UPropertyAnimatorFloatBase : public UPropertyAnimatorCoreBase
+UCLASS(MinimalAPI, Abstract, AutoExpandCategories=("Animator"))
+class UPropertyAnimatorFloatBase : public UPropertyAnimatorCoreBase
 {
 	GENERATED_BODY()
 
 	friend class FPropertyAnimatorCoreEditorDetailCustomization;
 
 public:
-	void SetGlobalMagnitude(float InMagnitude);
+	PROPERTYANIMATOR_API void SetGlobalMagnitude(float InMagnitude);
 	float GetGlobalMagnitude() const
 	{
 		return GlobalMagnitude;
 	}
 
-	void SetGlobalFrequency(float InFrequency);
+	PROPERTYANIMATOR_API void SetGlobalFrequency(float InFrequency);
 	float GetGlobalFrequency() const
 	{
 		return GlobalFrequency;
 	}
 
-	void SetAccumulatedTimeOffset(double InOffset);
+	PROPERTYANIMATOR_API void SetAccumulatedTimeOffset(double InOffset);
 	double GetAccumulatedTimeOffset() const
 	{
 		return AccumulatedTimeOffset;
 	}
 
-	void SetRandomTimeOffset(bool bInOffset);
+	PROPERTYANIMATOR_API void SetRandomTimeOffset(bool bInOffset);
 	bool GetRandomTimeOffset() const
 	{
 		return bRandomTimeOffset;
 	}
 
-	void SetSeed(int32 InSeed);
+	PROPERTYANIMATOR_API void SetSeed(int32 InSeed);
 	int32 GetSeed() const
 	{
 		return Seed;
diff --git a/Engine/Plugins/Experimental/PropertyAnimator/Source/PropertyAnimator/Public/Animators/PropertyAnimatorOscillate.h b/Engine/Plugins/Experimental/PropertyAnimator/Source/PropertyAnimator/Public/Animators/PropertyAnimatorOscillate.h
--- a/Engine/Plugins/Experimental/PropertyAnimator/Source/PropertyAnimator/Public/Animators/PropertyAnimatorOscillate.h
+++ b/Engine/Plugins/Experimental/PropertyAnimator/Source/PropertyAnimator/Public/Animators/PropertyAnimatorOscillate.h
@@ -19,8 +19,8 @@ enum class EPropertyAnimatorOscillateFunction : uint8
 /**
  * Applies an additive regular oscillate movement with various options on supported float properties
  */
-UCLASS(AutoExpandCategories=("Animator"))
-class PROPERTYANIMATOR_API UPropertyAnimatorOscillate : public UPropertyAnimatorFloatBase
+UCLASS(MinimalAPI, AutoExpandCategories=("Animator"))
+class UPropertyAnimatorOscillate : public UPropertyAnimatorFloatBase
 {
 	GENERATED_BODY()
 
@@ -29,7 +29,7 @@ public:
 
 	UPropertyAnimatorOscillate();
 
-	void SetOscillateFunction(EPropertyAnimatorOscillateFunction InFunction);
+	PROPERTYANIMATOR_API void SetOscillateFunction(EPropertyAnimatorOscillateFunction InFunction);
 	EPropertyAnimatorOscillateFunction GetOscillateFunction() const
 	{
 		return OscillateFunction;
diff --git a/Engine/Plugins/Experimental/PropertyAnimator/Source/PropertyAnimator/Public/Animators/PropertyAnimatorPulse.h b/Engine/Plugins/Experimental/PropertyAnimator/Source/PropertyAnimator/Public/Animators/PropertyAnimatorPulse.h
--- a/Engine/Plugins/Experimental/PropertyAnimator/Source/PropertyAnimator/Public/Animators/PropertyAnimatorPulse.h
+++ b/Engine/Plugins/Experimental/PropertyAnimator/Source/PropertyAnimator/Public/Animators/PropertyAnimatorPulse.h
@@ -11,8 +11,8 @@
 /**
  * Applies an additive pulse movement with various options on supported float properties
  */
-UCLASS(AutoExpandCategories=("Animator"))
-class PROPERTYANIMATOR_API UPropertyAnimatorPulse : public UPropertyAnimatorFloatBase
+UCLASS(MinimalAPI, AutoExpandCategories=("Animator"))
+class UPropertyAnimatorPulse : public UPropertyAnimatorFloatBase
 {
 	GENERATED_BODY()
 
@@ -21,13 +21,13 @@ public:
 
 	UPropertyAnimatorPulse();
 
-	void SetEasingFunction(EPropertyAnimatorEasingFunction InEasingFunction);
+	PROPERTYANIMATOR_API void SetEasingFunction(EPropertyAnimatorEasingFunction InEasingFunction);
 	EPropertyAnimatorEasingFunction GetEasingFunction() const
 	{
 		return EasingFunction;
 	}
 
-	void SetEasingType(EPropertyAnimatorEasingType InEasingType);
+	PROPERTYANIMATOR_API void SetEasingType(EPropertyAnimatorEasingType InEasingType);
 	EPropertyAnimatorEasingType GetEasingType() const
 	{
 		return EasingType;
diff --git a/Engine/Plugins/Experimental/PropertyAnimator/Source/PropertyAnimator/Public/Animators/PropertyAnimatorSoundWave.h b/Engine/Plugins/Experimental/PropertyAnimator/Source/PropertyAnimator/Public/Animators/PropertyAnimatorSoundWave.h
--- a/Engine/Plugins/Experimental/PropertyAnimator/Source/PropertyAnimator/Public/Animators/PropertyAnimatorSoundWave.h
+++ b/Engine/Plugins/Experimental/PropertyAnimator/Source/PropertyAnimator/Public/Animators/PropertyAnimatorSoundWave.h
@@ -11,8 +11,8 @@ class USoundWave;
 /**
  * Applies a sampled sound wave movement with various options on supported float properties
  */
-UCLASS(AutoExpandCategories=("Animator"))
-class PROPERTYANIMATOR_API UPropertyAnimatorSoundWave : public UPropertyAnimatorFloatBase
+UCLASS(MinimalAPI, AutoExpandCategories=("Animator"))
+class UPropertyAnimatorSoundWave : public UPropertyAnimatorFloatBase
 {
 	GENERATED_BODY()
 
@@ -21,7 +21,13 @@ public:
 
 	UPropertyAnimatorSoundWave();
 
-	void SetLoop(bool bInLoop);
+	PROPERTYANIMATOR_API void SetSampledSoundWave(USoundWave* InSoundWave);
+	USoundWave* GetSampledSoundWave() const
+	{
+		return SampledSoundWave;
+	}
+
+	PROPERTYANIMATOR_API void SetLoop(bool bInLoop);
 	bool GetLoop() const
 	{
 		return bLoop;
@@ -45,7 +51,7 @@ protected:
 	 * Cannot be switched at runtime, only in editor due to analyzer
 	 * Analyzed audio will work at runtime since it is cached
 	 */
-	UPROPERTY(EditInstanceOnly, BlueprintReadOnly, Category="Animator")
+	UPROPERTY(EditInstanceOnly, BlueprintReadOnly, Setter, Getter, Category="Animator")
 	TObjectPtr<USoundWave> SampledSoundWave;
 
 	/** Whether we keep looping after the duration has been reached or before 0 */
diff --git a/Engine/Plugins/Experimental/PropertyAnimator/Source/PropertyAnimator/Public/Animators/PropertyAnimatorTime.h b/Engine/Plugins/Experimental/PropertyAnimator/Source/PropertyAnimator/Public/Animators/PropertyAnimatorTime.h
--- a/Engine/Plugins/Experimental/PropertyAnimator/Source/PropertyAnimator/Public/Animators/PropertyAnimatorTime.h
+++ b/Engine/Plugins/Experimental/PropertyAnimator/Source/PropertyAnimator/Public/Animators/PropertyAnimatorTime.h
@@ -8,8 +8,8 @@
 /**
  * Applies an additive time movement with various options on supported float properties
  */
-UCLASS(AutoExpandCategories=("Animator"))
-class PROPERTYANIMATOR_API UPropertyAnimatorTime : public UPropertyAnimatorFloatBase
+UCLASS(MinimalAPI, AutoExpandCategories=("Animator"))
+class UPropertyAnimatorTime : public UPropertyAnimatorFloatBase
 {
 	GENERATED_BODY()
 
diff --git a/Engine/Plugins/Experimental/PropertyAnimator/Source/PropertyAnimator/Public/Animators/PropertyAnimatorWiggle.h b/Engine/Plugins/Experimental/PropertyAnimator/Source/PropertyAnimator/Public/Animators/PropertyAnimatorWiggle.h
--- a/Engine/Plugins/Experimental/PropertyAnimator/Source/PropertyAnimator/Public/Animators/PropertyAnimatorWiggle.h
+++ b/Engine/Plugins/Experimental/PropertyAnimator/Source/PropertyAnimator/Public/Animators/PropertyAnimatorWiggle.h
@@ -8,8 +8,8 @@
 /**
  * Applies an additive random wiggle movement with various options on supported float properties
  */
-UCLASS(AutoExpandCategories=("Animator"))
-class PROPERTYANIMATOR_API UPropertyAnimatorWiggle : public UPropertyAnimatorFloatBase
+UCLASS(MinimalAPI, AutoExpandCategories=("Animator"))
+class UPropertyAnimatorWiggle : public UPropertyAnimatorFloatBase
 {
 	GENERATED_BODY()
 
diff --git a/Engine/Plugins/Experimental/PropertyAnimator/Source/PropertyAnimator/Public/Properties/PropertyAnimatorFloatContext.h b/Engine/Plugins/Experimental/PropertyAnimator/Source/PropertyAnimator/Public/Properties/PropertyAnimatorFloatContext.h
--- a/Engine/Plugins/Experimental/PropertyAnimator/Source/PropertyAnimator/Public/Properties/PropertyAnimatorFloatContext.h
+++ b/Engine/Plugins/Experimental/PropertyAnimator/Source/PropertyAnimator/Public/Properties/PropertyAnimatorFloatContext.h
@@ -6,37 +6,37 @@
 #include "PropertyAnimatorFloatContext.generated.h"
 
 /** Property context used by float driving animator */
-UCLASS(BlueprintType)
-class PROPERTYANIMATOR_API UPropertyAnimatorFloatContext : public UPropertyAnimatorCoreContext
+UCLASS(MinimalAPI, BlueprintType)
+class UPropertyAnimatorFloatContext : public UPropertyAnimatorCoreContext
 {
 	GENERATED_BODY()
 
 public:
-	void SetMagnitude(float InMagnitude);
+	PROPERTYANIMATOR_API void SetMagnitude(float InMagnitude);
 	float GetMagnitude() const
 	{
 		return Magnitude;
 	}
 
-	void SetAmplitudeMin(double InAmplitude);
+	PROPERTYANIMATOR_API void SetAmplitudeMin(double InAmplitude);
 	double GetAmplitudeMin() const
 	{
 		return AmplitudeMin;
 	}
 
-	void SetAmplitudeMax(double InAmplitude);
+	PROPERTYANIMATOR_API void SetAmplitudeMax(double InAmplitude);
 	double GetAmplitudeMax() const
 	{
 		return AmplitudeMax;
 	}
 
-	void SetFrequency(float InFrequency);
+	PROPERTYANIMATOR_API void SetFrequency(float InFrequency);
 	float GetFrequency() const
 	{
 		return Frequency;
 	}
 
-	void SetTimeOffset(double InTimeOffset);
+	PROPERTYANIMATOR_API void SetTimeOffset(double InTimeOffset);
 	double GetTimeOffset() const
 	{
 		return TimeOffset;
diff --git a/Engine/Plugins/Experimental/PropertyAnimator/Source/PropertyAnimator/Public/Properties/PropertyAnimatorTextGroup.h b/Engine/Plugins/Experimental/PropertyAnimator/Source/PropertyAnimator/Public/Properties/PropertyAnimatorTextGroup.h
--- a/Engine/Plugins/Experimental/PropertyAnimator/Source/PropertyAnimator/Public/Properties/PropertyAnimatorTextGroup.h
+++ b/Engine/Plugins/Experimental/PropertyAnimator/Source/PropertyAnimator/Public/Properties/PropertyAnimatorTextGroup.h
@@ -6,8 +6,8 @@
 #include "PropertyAnimatorTextGroup.generated.h"
 
 /** Group that handles text character range properties */
-UCLASS(BlueprintType)
-class PROPERTYANIMATOR_API UPropertyAnimatorTextGroup : public UPropertyAnimatorCoreGroupBase
+UCLASS(MinimalAPI, BlueprintType)
+class UPropertyAnimatorTextGroup : public UPropertyAnimatorCoreGroupBase
 {
 	GENERATED_BODY()
 
@@ -16,18 +16,36 @@ public:
 		: UPropertyAnimatorCoreGroupBase(TEXT("Text"))
 	{}
 
+	PROPERTYANIMATOR_API void SetRangeStart(float InRangeStart);
+	float GetRangeStart() const
+	{
+		return RangeStart;
+	}
+
+	PROPERTYANIMATOR_API void SetRangeEnd(float InRangeEnd);
+	float GetRangeEnd() const
+	{
+		return RangeEnd;
+	}
+
+	PROPERTYANIMATOR_API void SetRangeOffset(float InRangeOffset);
+	float GetRangeOffset() const
+	{
+		return RangeOffset;
+	}
+
 	//~ Begin UPropertyAnimatorCoreGroupBase
 	virtual void ManageProperties(const UPropertyAnimatorCoreContext* InContext, TArray<FPropertyAnimatorCoreData>& InOutProperties) override;
 	virtual bool IsPropertySupported(const UPropertyAnimatorCoreContext* InContext) const override;
 	//~ End UPropertyAnimatorCoreGroupBase
 
 protected:
-	UPROPERTY(EditInstanceOnly, BlueprintReadWrite, Interp, Category="Animator", meta=(ClampMin="0", ClampMax="1"))
+	UPROPERTY(EditInstanceOnly, BlueprintReadWrite, Interp, Setter, Getter, Category="Animator", meta=(ClampMin="0", ClampMax="1"))
 	float RangeStart = 0.f;
 
-	UPROPERTY(EditInstanceOnly, BlueprintReadWrite, Interp, Category="Animator", meta=(ClampMin="0", ClampMax="1"))
+	UPROPERTY(EditInstanceOnly, BlueprintReadWrite, Interp, Setter, Getter, Category="Animator", meta=(ClampMin="0", ClampMax="1"))
 	float RangeEnd = 100.f;
 
-	UPROPERTY(EditInstanceOnly, BlueprintReadWrite, Interp, Category="Animator")
+	UPROPERTY(EditInstanceOnly, BlueprintReadWrite, Interp, Setter, Getter, Category="Animator")
 	float RangeOffset = 0.f;
 };
\ No newline at end of file
diff --git a/Engine/Plugins/Experimental/PropertyAnimatorCore/Source/PropertyAnimatorCore/Private/Subsystems/PropertyAnimatorCoreSubsystem.cpp b/Engine/Plugins/Experimental/PropertyAnimatorCore/Source/PropertyAnimatorCore/Private/Subsystems/PropertyAnimatorCoreSubsystem.cpp
--- a/Engine/Plugins/Experimental/PropertyAnimatorCore/Source/PropertyAnimatorCore/Private/Subsystems/PropertyAnimatorCoreSubsystem.cpp
+++ b/Engine/Plugins/Experimental/PropertyAnimatorCore/Source/PropertyAnimatorCore/Private/Subsystems/PropertyAnimatorCoreSubsystem.cpp
@@ -1174,6 +1174,11 @@ bool UPropertyAnimatorCoreSubsystem::UnregisterSetterResolver(FName InPropertyNa
 	return SetterResolvers.Remove(InPropertyName) > 0;
 }
 
+bool UPropertyAnimatorCoreSubsystem::IsSetterResolverRegistered(FName InPropertyName) const
+{
+	return SetterResolvers.Contains(InPropertyName);
+}
+
 UFunction* UPropertyAnimatorCoreSubsystem::ResolveSetter(FName InPropertyName, const UObject* InOwner)
 {
 	if (!IsValid(InOwner))
diff --git a/Engine/Plugins/Experimental/PropertyAnimatorCore/Source/PropertyAnimatorCore/Private/TimeSources/PropertyAnimatorCoreManualTimeSource.cpp b/Engine/Plugins/Experimental/PropertyAnimatorCore/Source/PropertyAnimatorCore/Private/TimeSources/PropertyAnimatorCoreManualTimeSource.cpp
--- a/Engine/Plugins/Experimental/PropertyAnimatorCore/Source/PropertyAnimatorCore/Private/TimeSources/PropertyAnimatorCoreManualTimeSource.cpp
+++ b/Engine/Plugins/Experimental/PropertyAnimatorCore/Source/PropertyAnimatorCore/Private/TimeSources/PropertyAnimatorCoreManualTimeSource.cpp
@@ -2,7 +2,7 @@
 
 #include "TimeSources/PropertyAnimatorCoreManualTimeSource.h"
 
-void UPropertyAnimatorCoreManualTimeSource::SetCustomTime(float InTime)
+void UPropertyAnimatorCoreManualTimeSource::SetCustomTime(double InTime)
 {
 	if (FMath::IsNearlyEqual(InTime, CustomTime))
 	{
diff --git a/Engine/Plugins/Experimental/PropertyAnimatorCore/Source/PropertyAnimatorCore/Public/Animators/PropertyAnimatorCoreBase.h b/Engine/Plugins/Experimental/PropertyAnimatorCore/Source/PropertyAnimatorCore/Public/Animators/PropertyAnimatorCoreBase.h
--- a/Engine/Plugins/Experimental/PropertyAnimatorCore/Source/PropertyAnimatorCore/Public/Animators/PropertyAnimatorCoreBase.h
+++ b/Engine/Plugins/Experimental/PropertyAnimatorCore/Source/PropertyAnimatorCore/Public/Animators/PropertyAnimatorCoreBase.h
@@ -22,8 +22,8 @@ struct FPropertyAnimatorCoreEvaluationParameters
 };
 
 /** Abstract base class for any Animator, holds a set of linked properties */
-UCLASS(Abstract, EditInlineNew, AutoExpandCategories=("Animator"))
-class PROPERTYANIMATORCORE_API UPropertyAnimatorCoreBase : public UObject
+UCLASS(MinimalAPI, Abstract, EditInlineNew, AutoExpandCategories=("Animator"))
+class UPropertyAnimatorCoreBase : public UObject
 {
 	GENERATED_BODY()
 
@@ -36,39 +36,40 @@ public:
 	DECLARE_MULTICAST_DELEGATE_OneParam(FOnAnimatorUpdated, UPropertyAnimatorCoreBase* /* InAnimator */)
 
 	/** Called when a Animator is created */
-	static FOnAnimatorUpdated OnAnimatorCreatedDelegate;
+	PROPERTYANIMATORCORE_API static FOnAnimatorUpdated OnAnimatorCreatedDelegate;
 
 	/** Called when a Animator is removed */
-	static FOnAnimatorUpdated OnAnimatorRemovedDelegate;
+	PROPERTYANIMATORCORE_API static FOnAnimatorUpdated OnAnimatorRemovedDelegate;
 
 	/** Called when a Animator is renamed */
-	static FOnAnimatorUpdated OnAnimatorRenamedDelegate;
+	PROPERTYANIMATORCORE_API static FOnAnimatorUpdated OnAnimatorRenamedDelegate;
 
 	DECLARE_MULTICAST_DELEGATE_TwoParams(FOnAnimatorPropertyUpdated, UPropertyAnimatorCoreBase* /* InAnimator */, const FPropertyAnimatorCoreData& /** InProperty */)
 
 	/** Called when a property is linked to a Animator */
-	static FOnAnimatorPropertyUpdated OnAnimatorPropertyLinkedDelegate;
+	PROPERTYANIMATORCORE_API static FOnAnimatorPropertyUpdated OnAnimatorPropertyLinkedDelegate;
 
 	/** Called when a property is unlinked to a Animator */
-	static FOnAnimatorPropertyUpdated OnAnimatorPropertyUnlinkedDelegate;
+	PROPERTYANIMATORCORE_API static FOnAnimatorPropertyUpdated OnAnimatorPropertyUnlinkedDelegate;
 
-	UPropertyAnimatorCoreBase();
+	PROPERTYANIMATORCORE_API UPropertyAnimatorCoreBase();
+
+	PROPERTYANIMATORCORE_API AActor* GetAnimatorActor() const;
 
-	AActor* GetAnimatorActor() const;
 	UPropertyAnimatorCoreComponent* GetAnimatorComponent() const;
 
 	/** Get the global animators magnitude */
 	float GetAnimatorComponentMagnitude() const;
 
 	/** Set the state of this animator */
-	void SetAnimatorEnabled(bool bInIsEnabled);
+	PROPERTYANIMATORCORE_API void SetAnimatorEnabled(bool bInIsEnabled);
 	bool GetAnimatorEnabled() const
 	{
 		return bAnimatorEnabled;
 	}
 
 	/** Set the time source name to use */
-	void SetTimeSourceName(FName InTimeSourceName);
+	PROPERTYANIMATORCORE_API void SetTimeSourceName(FName InTimeSourceName);
 	FName GetTimeSourceName() const
 	{
 		return TimeSourceName;
@@ -81,45 +82,45 @@ public:
 	}
 
 	/** Set the display name of this animator */
-	void SetAnimatorDisplayName(FName InName);
+	PROPERTYANIMATORCORE_API void SetAnimatorDisplayName(FName InName);
 	FString GetAnimatorDisplayName() const
 	{
 		return AnimatorDisplayName.ToString();
 	}
 
 	/** Gets the Animator original name */
-	FName GetAnimatorOriginalName() const;
+	PROPERTYANIMATORCORE_API FName GetAnimatorOriginalName() const;
 
 	/** Get all linked properties within this animator */
-	TSet<FPropertyAnimatorCoreData> GetLinkedProperties() const;
+	PROPERTYANIMATORCORE_API TSet<FPropertyAnimatorCoreData> GetLinkedProperties() const;
 
 	/** Get linked properties count within this animator */
-	int32 GetLinkedPropertiesCount() const;
+	PROPERTYANIMATORCORE_API int32 GetLinkedPropertiesCount() const;
 
 	/** Link property to this Animator to be able to drive it */
-	bool LinkProperty(const FPropertyAnimatorCoreData& InLinkProperty);
+	PROPERTYANIMATORCORE_API bool LinkProperty(const FPropertyAnimatorCoreData& InLinkProperty);
 
 	/** Unlink property from this Animator */
-	bool UnlinkProperty(const FPropertyAnimatorCoreData& InUnlinkProperty);
+	PROPERTYANIMATORCORE_API bool UnlinkProperty(const FPropertyAnimatorCoreData& InUnlinkProperty);
 
 	/** Checks if this Animator is controlling this property */
-	bool IsPropertyLinked(const FPropertyAnimatorCoreData& InPropertyData) const;
+	PROPERTYANIMATORCORE_API bool IsPropertyLinked(const FPropertyAnimatorCoreData& InPropertyData) const;
 
 	/** Checks if this animator is controlling all properties */
 	bool IsPropertiesLinked(const TSet<FPropertyAnimatorCoreData>& InProperties) const;
 
 	/** Returns all inner properties that are controlled by this Animator linked to member property */
-	TSet<FPropertyAnimatorCoreData> GetInnerPropertiesLinked(const FPropertyAnimatorCoreData& InPropertyData) const;
+	PROPERTYANIMATORCORE_API TSet<FPropertyAnimatorCoreData> GetInnerPropertiesLinked(const FPropertyAnimatorCoreData& InPropertyData) const;
 
 	/**
 	 * Checks recursively for properties inside that member property that are supported by this Animator
 	 * if the member property is directly supported only returns that one, calls IsPropertySupported to check
 	 * When recursive is false, it will stop if member property or children is controllable and not look further in the chain
 	 */
-	bool GetPropertiesSupported(const FPropertyAnimatorCoreData& InPropertyData, TSet<FPropertyAnimatorCoreData>& OutProperties, bool bInRecursiveSearch = false) const;
+	PROPERTYANIMATORCORE_API bool GetPropertiesSupported(const FPropertyAnimatorCoreData& InPropertyData, TSet<FPropertyAnimatorCoreData>& OutProperties, bool bInRecursiveSearch = false) const;
 
 	/** Does not recurse, checks if this Animator supports this specific property, use GetPropertiesSupported to check if nested properties are supported */
-	bool IsPropertySupported(const FPropertyAnimatorCoreData& InPropertyData) const;
+	PROPERTYANIMATORCORE_API bool IsPropertySupported(const FPropertyAnimatorCoreData& InPropertyData) const;
 
 	/** Override this to check if a property is directly supported by this animator */
 	virtual bool IsPropertyDirectlySupported(const FPropertyAnimatorCoreData& InPropertyData) const
@@ -134,7 +135,7 @@ public:
 	}
 
 	/** Get the context for the linked property */
-	UPropertyAnimatorCoreContext* GetLinkedPropertyContext(const FPropertyAnimatorCoreData& InProperty) const;
+	PROPERTYANIMATORCORE_API UPropertyAnimatorCoreContext* GetLinkedPropertyContext(const FPropertyAnimatorCoreData& InProperty) const;
 
 	/** Get the casted context for the linked property */
 	template<typename InContextClass
@@ -146,11 +147,11 @@ public:
 
 protected:
 	//~ Begin UObject
-	virtual void PostLoad() override;
+	PROPERTYANIMATORCORE_API virtual void PostLoad() override;
 #if WITH_EDITOR
-	virtual void PreEditUndo() override;
-	virtual void PostEditUndo() override;
-	virtual void PostEditChangeProperty(FPropertyChangedEvent& PropertyChangedEvent) override;
+	PROPERTYANIMATORCORE_API virtual void PreEditUndo() override;
+	PROPERTYANIMATORCORE_API virtual void PostEditUndo() override;
+	PROPERTYANIMATORCORE_API virtual void PostEditChangeProperty(FPropertyChangedEvent& PropertyChangedEvent) override;
 #endif
 	//~ End UObject
 
@@ -236,24 +237,25 @@ protected:
 	virtual void OnAnimatorAdded() {}
 	virtual void OnAnimatorRemoved() {}
 
-	virtual void OnAnimatorEnabled();
-	virtual void OnAnimatorDisabled();
+	PROPERTYANIMATORCORE_API virtual void OnAnimatorEnabled();
+	PROPERTYANIMATORCORE_API virtual void OnAnimatorDisabled();
 
 	/** Returns the property context class to use */
-	virtual TSubclassOf<UPropertyAnimatorCoreContext> GetPropertyContextClass(const FPropertyAnimatorCoreData& InProperty);
+	PROPERTYANIMATORCORE_API virtual TSubclassOf<UPropertyAnimatorCoreContext> GetPropertyContextClass(const FPropertyAnimatorCoreData& InProperty);
+
 	virtual void OnPropertyLinked(UPropertyAnimatorCoreContext* InLinkedProperty) {}
 	virtual void OnPropertyUnlinked(UPropertyAnimatorCoreContext* InUnlinkedProperty) {}
 
-	/** Restore modified properties to original state */
-	virtual void RestoreProperties(bool bInForce = false);
-
-	/** Allocate and saves properties in the property bag */
-	virtual void SaveProperties();
-
 	/** Apply animators effect on linked properties */
 	virtual void EvaluateProperties(const FPropertyAnimatorCoreEvaluationParameters& InParameters) {}
 
 private:
+	/** Restore modified properties to original state */
+	void RestoreProperties(bool bInForce = false);
+
+	/** Allocate and saves properties in the property bag */
+	void SaveProperties();
+
 	/** Called by the component to evaluate this animator */
 	void EvaluateAnimator();
 
diff --git a/Engine/Plugins/Experimental/PropertyAnimatorCore/Source/PropertyAnimatorCore/Public/Components/PropertyAnimatorCoreComponent.h b/Engine/Plugins/Experimental/PropertyAnimatorCore/Source/PropertyAnimatorCore/Public/Components/PropertyAnimatorCoreComponent.h
--- a/Engine/Plugins/Experimental/PropertyAnimatorCore/Source/PropertyAnimatorCore/Public/Components/PropertyAnimatorCoreComponent.h
+++ b/Engine/Plugins/Experimental/PropertyAnimatorCore/Source/PropertyAnimatorCore/Public/Components/PropertyAnimatorCoreComponent.h
@@ -7,8 +7,8 @@
 #include "PropertyAnimatorCoreComponent.generated.h"
 
 /** A container for controllers that holds properties in this actor */
-UCLASS(ClassGroup=(Custom), AutoExpandCategories=("Animator"), HideCategories=("Activation", "Cooking", "AssetUserData", "Collision"), meta=(BlueprintSpawnableComponent))
-class PROPERTYANIMATORCORE_API UPropertyAnimatorCoreComponent : public UActorComponent
+UCLASS(MinimalAPI, ClassGroup=(Custom), AutoExpandCategories=("Animator"), HideCategories=("Activation", "Cooking", "AssetUserData", "Collision"), meta=(BlueprintSpawnableComponent))
+class UPropertyAnimatorCoreComponent : public UActorComponent
 {
 	GENERATED_BODY()
 
diff --git a/Engine/Plugins/Experimental/PropertyAnimatorCore/Source/PropertyAnimatorCore/Public/Presets/PropertyAnimatorCorePresetBase.h b/Engine/Plugins/Experimental/PropertyAnimatorCore/Source/PropertyAnimatorCore/Public/Presets/PropertyAnimatorCorePresetBase.h
--- a/Engine/Plugins/Experimental/PropertyAnimatorCore/Source/PropertyAnimatorCore/Public/Presets/PropertyAnimatorCorePresetBase.h
+++ b/Engine/Plugins/Experimental/PropertyAnimatorCore/Source/PropertyAnimatorCore/Public/Presets/PropertyAnimatorCorePresetBase.h
@@ -19,15 +19,15 @@ class UPropertyAnimatorCorePresetBase : public UObject
 	GENERATED_BODY()
 
 public:
-	PROPERTYANIMATORCORE_API UPropertyAnimatorCorePresetBase()
+	UPropertyAnimatorCorePresetBase()
 		: UPropertyAnimatorCorePresetBase(NAME_None)
 	{}
 
-	PROPERTYANIMATORCORE_API UPropertyAnimatorCorePresetBase(FName InPresetName)
+	UPropertyAnimatorCorePresetBase(FName InPresetName)
 		: PresetName(InPresetName)
 	{}
 
-	PROPERTYANIMATORCORE_API FName GetPresetName() const
+	FName GetPresetName() const
 	{
 		return PresetName;
 	}
@@ -35,7 +35,7 @@ public:
 	PROPERTYANIMATORCORE_API FString GetPresetDisplayName() const;
 
 	/** Get the preset properties for that actor */
-	PROPERTYANIMATORCORE_API virtual void GetPresetProperties(const AActor* InActor, const UPropertyAnimatorCoreBase* InAnimator, TSet<FPropertyAnimatorCoreData>& OutProperties) const {}
+	virtual void GetPresetProperties(const AActor* InActor, const UPropertyAnimatorCoreBase* InAnimator, TSet<FPropertyAnimatorCoreData>& OutProperties) const {}
 
 	/** Get the preset properties for that actor but only supported ones for that animator */
 	PROPERTYANIMATORCORE_API virtual void GetSupportedPresetProperties(const AActor* InActor, const UPropertyAnimatorCoreBase* InAnimator, TSet<FPropertyAnimatorCoreData>& OutProperties) const;
@@ -53,7 +53,7 @@ public:
 	PROPERTYANIMATORCORE_API virtual bool UnapplyPreset(UPropertyAnimatorCoreBase* InAnimator);
 
 	/** Called when this preset is applied on the animator */
-	PROPERTYANIMATORCORE_API virtual void OnPresetApplied(UPropertyAnimatorCoreBase* InAnimator, const TSet<FPropertyAnimatorCoreData>& InProperties) {}
+	virtual void OnPresetApplied(UPropertyAnimatorCoreBase* InAnimator, const TSet<FPropertyAnimatorCoreData>& InProperties) {}
 
 private:
 	/** Name used to display this preset to the user */
diff --git a/Engine/Plugins/Experimental/PropertyAnimatorCore/Source/PropertyAnimatorCore/Public/Properties/PropertyAnimatorCoreContext.h b/Engine/Plugins/Experimental/PropertyAnimatorCore/Source/PropertyAnimatorCore/Public/Properties/PropertyAnimatorCoreContext.h
--- a/Engine/Plugins/Experimental/PropertyAnimatorCore/Source/PropertyAnimatorCore/Public/Properties/PropertyAnimatorCoreContext.h
+++ b/Engine/Plugins/Experimental/PropertyAnimatorCore/Source/PropertyAnimatorCore/Public/Properties/PropertyAnimatorCoreContext.h
@@ -36,10 +36,10 @@ public:
 		return AnimatedProperty;
 	}
 
-	PROPERTYANIMATORCORE_API UPropertyAnimatorCoreBase* GetAnimator() const;
+	UPropertyAnimatorCoreBase* GetAnimator() const;
 
 	/** Get the handler responsible for this property type */
-	PROPERTYANIMATORCORE_API UPropertyAnimatorCoreHandlerBase* GetHandler() const;
+	UPropertyAnimatorCoreHandlerBase* GetHandler() const;
 
 	/** Get the active group of this property */
 	UPropertyAnimatorCoreGroupBase* GetGroup() const
diff --git a/Engine/Plugins/Experimental/PropertyAnimatorCore/Source/PropertyAnimatorCore/Public/Properties/PropertyAnimatorCoreData.h b/Engine/Plugins/Experimental/PropertyAnimatorCore/Source/PropertyAnimatorCore/Public/Properties/PropertyAnimatorCoreData.h
--- a/Engine/Plugins/Experimental/PropertyAnimatorCore/Source/PropertyAnimatorCore/Public/Properties/PropertyAnimatorCoreData.h
+++ b/Engine/Plugins/Experimental/PropertyAnimatorCore/Source/PropertyAnimatorCore/Public/Properties/PropertyAnimatorCoreData.h
@@ -20,26 +20,26 @@ class UPropertyAnimatorCoreResolver;
 
 /** Serializable struct that contains the property and the owner with accessors */
 USTRUCT(BlueprintType)
-struct PROPERTYANIMATORCORE_API FPropertyAnimatorCoreData
+struct FPropertyAnimatorCoreData
 {
 	GENERATED_BODY()
 
 	FPropertyAnimatorCoreData() = default;
 
 	/** Takes the owner, the member property and the inner property inside */
-	explicit FPropertyAnimatorCoreData(UObject* InObject, FProperty* InMemberProperty, FProperty* InProperty, TSubclassOf<UPropertyAnimatorCoreResolver> InResolverClass = nullptr);
+	PROPERTYANIMATORCORE_API explicit FPropertyAnimatorCoreData(UObject* InObject, FProperty* InMemberProperty, FProperty* InProperty, TSubclassOf<UPropertyAnimatorCoreResolver> InResolverClass = nullptr);
 
 	/** Take the owner and the full property chain, from member property to inner property */
-	explicit FPropertyAnimatorCoreData(UObject* InObject, const TArray<FProperty*>& InChainProperties, TSubclassOf<UPropertyAnimatorCoreResolver> InResolverClass = nullptr);
+	PROPERTYANIMATORCORE_API explicit FPropertyAnimatorCoreData(UObject* InObject, const TArray<FProperty*>& InChainProperties, TSubclassOf<UPropertyAnimatorCoreResolver> InResolverClass = nullptr);
 
 	/** Take the owner, the property chain until the inner property and lastly the inner property */
-	explicit FPropertyAnimatorCoreData(UObject* InObject, const TArray<FProperty*>& InChainProperties, FProperty* InProperty, TSubclassOf<UPropertyAnimatorCoreResolver> InResolverClass = nullptr);
+	PROPERTYANIMATORCORE_API explicit FPropertyAnimatorCoreData(UObject* InObject, const TArray<FProperty*>& InChainProperties, FProperty* InProperty, TSubclassOf<UPropertyAnimatorCoreResolver> InResolverClass = nullptr);
 
 	/** Is this a resolvable property that uses a custom resolver */
-	bool IsResolvable() const;
+	PROPERTYANIMATORCORE_API bool IsResolvable() const;
 
 	/** Get the linked property resolver for resolvable properties */
-	UPropertyAnimatorCoreResolver* GetPropertyResolver() const;
+	PROPERTYANIMATORCORE_API UPropertyAnimatorCoreResolver* GetPropertyResolver() const;
 
 	/** Get the property resolver class */
 	TSubclassOf<UPropertyAnimatorCoreResolver> GetPropertyResolverClass() const;
@@ -61,10 +61,10 @@ struct PROPERTYANIMATORCORE_API FPropertyAnimatorCoreData
 	}
 
 	/** Returns the owning actor from this owner object */
-	AActor* GetOwningActor() const;
+	PROPERTYANIMATORCORE_API AActor* GetOwningActor() const;
 
 	/** Returns the owning component from this owner object if any */
-	UActorComponent* GetOwningComponent() const;
+	PROPERTYANIMATORCORE_API UActorComponent* GetOwningComponent() const;
 
 	/** The owner object of the member property */
 	UObject* GetOwner() const
@@ -97,10 +97,10 @@ struct PROPERTYANIMATORCORE_API FPropertyAnimatorCoreData
 	}
 
 	/** The member property name */
-	FName GetMemberPropertyName() const;
+	PROPERTYANIMATORCORE_API FName GetMemberPropertyName() const;
 
 	/** The leaf property name */
-	FName GetLeafPropertyName() const;
+	PROPERTYANIMATORCORE_API FName GetLeafPropertyName() const;
 
 	/** The chain properties from member to inner property */
 	const TArray<TFieldPath<FProperty>>& GetChainProperties() const
@@ -109,7 +109,7 @@ struct PROPERTYANIMATORCORE_API FPropertyAnimatorCoreData
 	}
 
 	/** Checks if the property is settable via setter */
-	bool HasSetter() const;
+	PROPERTYANIMATORCORE_API bool HasSetter() const;
 
 	/** Checks if we contain this other property directly */
 	bool IsParentOf(const FPropertyAnimatorCoreData& InOtherProperty) const;
@@ -124,13 +124,13 @@ struct PROPERTYANIMATORCORE_API FPropertyAnimatorCoreData
 	bool IsTransient() const;
 
 	/** Tries to find, based on this property, the direct child of other property */
-	TOptional<FPropertyAnimatorCoreData> GetChildOf(const FPropertyAnimatorCoreData& InOtherProperty) const;
+	PROPERTYANIMATORCORE_API TOptional<FPropertyAnimatorCoreData> GetChildOf(const FPropertyAnimatorCoreData& InOtherProperty) const;
 
 	/** Returns the parent of this property if there is one */
 	TOptional<FPropertyAnimatorCoreData> GetParent() const;
 
 	/** Returns the top most parent / member property if there is one */
-	TOptional<FPropertyAnimatorCoreData> GetRootParent() const;
+	PROPERTYANIMATORCORE_API TOptional<FPropertyAnimatorCoreData> GetRootParent() const;
 
 	template<typename InPropertyClass
 		UE_REQUIRES(std::is_base_of_v<FProperty, InPropertyClass>)>
diff --git a/Engine/Plugins/Experimental/PropertyAnimatorCore/Source/PropertyAnimatorCore/Public/Properties/PropertyAnimatorCoreResolver.h b/Engine/Plugins/Experimental/PropertyAnimatorCore/Source/PropertyAnimatorCore/Public/Properties/PropertyAnimatorCoreResolver.h
--- a/Engine/Plugins/Experimental/PropertyAnimatorCore/Source/PropertyAnimatorCore/Public/Properties/PropertyAnimatorCoreResolver.h
+++ b/Engine/Plugins/Experimental/PropertyAnimatorCore/Source/PropertyAnimatorCore/Public/Properties/PropertyAnimatorCoreResolver.h
@@ -13,8 +13,8 @@
  * that we cannot reach or are transient, will be resolved when needed
  * Resolvers should remain transient and stateless
  */
-UCLASS(Abstract, Transient)
-class PROPERTYANIMATORCORE_API UPropertyAnimatorCoreResolver : public UObject
+UCLASS(MinimalAPI, Abstract, Transient)
+class UPropertyAnimatorCoreResolver : public UObject
 {
 	GENERATED_BODY()
 
diff --git a/Engine/Plugins/Experimental/PropertyAnimatorCore/Source/PropertyAnimatorCore/Public/Subsystems/PropertyAnimatorCoreSubsystem.h b/Engine/Plugins/Experimental/PropertyAnimatorCore/Source/PropertyAnimatorCore/Public/Subsystems/PropertyAnimatorCoreSubsystem.h
--- a/Engine/Plugins/Experimental/PropertyAnimatorCore/Source/PropertyAnimatorCore/Public/Subsystems/PropertyAnimatorCoreSubsystem.h
+++ b/Engine/Plugins/Experimental/PropertyAnimatorCore/Source/PropertyAnimatorCore/Public/Subsystems/PropertyAnimatorCoreSubsystem.h
@@ -19,7 +19,7 @@ class UPropertyAnimatorCoreTimeSourceBase;
 
 /** This subsystem handle all property animators */
 UCLASS()
-class PROPERTYANIMATORCORE_API UPropertyAnimatorCoreSubsystem : public UEngineSubsystem
+class UPropertyAnimatorCoreSubsystem : public UEngineSubsystem
 {
 	GENERATED_BODY()
 
@@ -27,7 +27,7 @@ class PROPERTYANIMATORCORE_API UPropertyAnimatorCoreSubsystem : public UEngineSu
 
 public:
 	/** Get this subsystem instance */
-	static UPropertyAnimatorCoreSubsystem* Get();
+	PROPERTYANIMATORCORE_API static UPropertyAnimatorCoreSubsystem* Get();
 
 	//~ Begin USubsystem
 	virtual void Initialize(FSubsystemCollectionBase& Collection) override;
@@ -35,63 +35,63 @@ public:
 	//~ End USubsystem
 
 	/** Register the property controller class to allow its usage */
-	bool RegisterAnimatorClass(const UClass* InPropertyControllerClass);
+	PROPERTYANIMATORCORE_API bool RegisterAnimatorClass(const UClass* InPropertyControllerClass);
 
 	/** Unregister the property controller class to disallow its usage */
-	bool UnregisterAnimatorClass(const UClass* InPropertyControllerClass);
+	PROPERTYANIMATORCORE_API bool UnregisterAnimatorClass(const UClass* InPropertyControllerClass);
 
 	/** Checks if the property controller class is already registered */
-	bool IsAnimatorClassRegistered(const UClass* InPropertyControllerClass) const;
+	PROPERTYANIMATORCORE_API bool IsAnimatorClassRegistered(const UClass* InPropertyControllerClass) const;
 
 	/** Gets the animator CDO registered from the class */
 	UPropertyAnimatorCoreBase* GetAnimatorRegistered(const UClass* InAnimatorClass) const;
 
 	/** Returns true if any controller is able to control that property or nested otherwise false */
-	bool IsPropertySupported(const FPropertyAnimatorCoreData& InPropertyData, bool bInCheckNestedProperties = true) const;
+	PROPERTYANIMATORCORE_API bool IsPropertySupported(const FPropertyAnimatorCoreData& InPropertyData, bool bInCheckNestedProperties = true) const;
 
 	/** Find all animators linked to the property */
-	TSet<UPropertyAnimatorCoreBase*> GetPropertyLinkedAnimators(const FPropertyAnimatorCoreData& InPropertyData) const;
+	PROPERTYANIMATORCORE_API TSet<UPropertyAnimatorCoreBase*> GetPropertyLinkedAnimators(const FPropertyAnimatorCoreData& InPropertyData) const;
 
 	/** Returns a set of existing property controller objects in owner that supports that property */
-	TSet<UPropertyAnimatorCoreBase*> GetExistingAnimators(const FPropertyAnimatorCoreData& InPropertyData) const;
-	TSet<UPropertyAnimatorCoreBase*> GetExistingAnimators(const AActor* InActor) const;
+	PROPERTYANIMATORCORE_API TSet<UPropertyAnimatorCoreBase*> GetExistingAnimators(const FPropertyAnimatorCoreData& InPropertyData) const;
+	PROPERTYANIMATORCORE_API TSet<UPropertyAnimatorCoreBase*> GetExistingAnimators(const AActor* InActor) const;
 
 	/** Returns a set of property controller CDO that supports that property */
-	TSet<UPropertyAnimatorCoreBase*> GetAvailableAnimators(const FPropertyAnimatorCoreData* InPropertyData) const;
-	TSet<UPropertyAnimatorCoreBase*> GetAvailableAnimators() const;
+	PROPERTYANIMATORCORE_API TSet<UPropertyAnimatorCoreBase*> GetAvailableAnimators(const FPropertyAnimatorCoreData* InPropertyData) const;
+	PROPERTYANIMATORCORE_API TSet<UPropertyAnimatorCoreBase*> GetAvailableAnimators() const;
 
 	/** Register the property handler class to allow its usage */
-	bool RegisterHandlerClass(const UClass* InHandlerClass);
+	PROPERTYANIMATORCORE_API bool RegisterHandlerClass(const UClass* InHandlerClass);
 
 	/** Unregister the property handler class to disallow its usage */
-	bool UnregisterHandlerClass(const UClass* InHandlerClass);
+	PROPERTYANIMATORCORE_API bool UnregisterHandlerClass(const UClass* InHandlerClass);
 
 	/** Checks if the property handler class is already registered */
-	bool IsHandlerClassRegistered(const UClass* InHandlerClass) const;
+	PROPERTYANIMATORCORE_API bool IsHandlerClassRegistered(const UClass* InHandlerClass) const;
 
 	/** Gets a property handler for this property */
 	UPropertyAnimatorCoreHandlerBase* GetHandler(const FPropertyAnimatorCoreData& InPropertyData) const;
 
 	/** Register a resolver for custom properties */
-	bool RegisterResolverClass(const UClass* InResolverClass);
+	PROPERTYANIMATORCORE_API bool RegisterResolverClass(const UClass* InResolverClass);
 
 	/** Unregister a resolver */
-	bool UnregisterResolverClass(const UClass* InResolverClass);
+	PROPERTYANIMATORCORE_API bool UnregisterResolverClass(const UClass* InResolverClass);
 
 	/** Is this resolver registered */
-	bool IsResolverClassRegistered(const UClass* InResolverClass) const;
+	PROPERTYANIMATORCORE_API bool IsResolverClassRegistered(const UClass* InResolverClass) const;
 
 	/** Get all resolvable properties for an object property */
-	void GetResolvableProperties(const FPropertyAnimatorCoreData& InPropertyData, TSet<FPropertyAnimatorCoreData>& OutProperties) const;
+	PROPERTYANIMATORCORE_API void GetResolvableProperties(const FPropertyAnimatorCoreData& InPropertyData, TSet<FPropertyAnimatorCoreData>& OutProperties) const;
 
 	/** Register a time source class to control clock for animators */
-	bool RegisterTimeSourceClass(UClass* InTimeSourceClass);
+	PROPERTYANIMATORCORE_API bool RegisterTimeSourceClass(UClass* InTimeSourceClass);
 
 	/** Unregister a time source class */
-	bool UnregisterTimeSourceClass(UClass* InTimeSourceClass);
+	PROPERTYANIMATORCORE_API bool UnregisterTimeSourceClass(UClass* InTimeSourceClass);
 
 	/** Check time source class is registered */
-	bool IsTimeSourceClassRegistered(UClass* InTimeSourceClass) const;
+	PROPERTYANIMATORCORE_API bool IsTimeSourceClassRegistered(UClass* InTimeSourceClass) const;
 
 	/** Get all time sources available */
 	TArray<FName> GetTimeSourceNames() const;
@@ -103,75 +103,77 @@ public:
 	UPropertyAnimatorCoreTimeSourceBase* CreateNewTimeSource(FName InTimeSourceName, UPropertyAnimatorCoreBase* InAnimator);
 
 	/** Register a preset class */
-	bool RegisterPresetClass(const UClass* InPresetClass);
+	PROPERTYANIMATORCORE_API bool RegisterPresetClass(const UClass* InPresetClass);
 
 	/** Unregister a preset class */
-	bool UnregisterPresetClass(const UClass* InPresetClass);
+	PROPERTYANIMATORCORE_API bool UnregisterPresetClass(const UClass* InPresetClass);
 
 	/** Is this preset class registered */
-	bool IsPresetClassRegistered(const UClass* InPresetClass) const;
+	PROPERTYANIMATORCORE_API bool IsPresetClassRegistered(const UClass* InPresetClass) const;
 
 	/** Gets all supported presets for a specific animator and actor */
-	TSet<UPropertyAnimatorCorePresetBase*> GetSupportedPresets(const AActor* InActor, const UPropertyAnimatorCoreBase* InAnimator) const;
+	PROPERTYANIMATORCORE_API TSet<UPropertyAnimatorCorePresetBase*> GetSupportedPresets(const AActor* InActor, const UPropertyAnimatorCoreBase* InAnimator) const;
 
 	/** Get all registered preset available */
-	TSet<UPropertyAnimatorCorePresetBase*> GetAvailablePresets() const;
+	PROPERTYANIMATORCORE_API TSet<UPropertyAnimatorCorePresetBase*> GetAvailablePresets() const;
 
-	bool RegisterSetterResolver(FName InPropertyName, TFunction<UFunction*(const UObject*)>&& InFunction);
+	PROPERTYANIMATORCORE_API bool RegisterSetterResolver(FName InPropertyName, TFunction<UFunction*(const UObject*)>&& InFunction);
 
-	bool UnregisterSetterResolver(FName InPropertyName);
+	PROPERTYANIMATORCORE_API bool UnregisterSetterResolver(FName InPropertyName);
+
+	PROPERTYANIMATORCORE_API bool IsSetterResolverRegistered(FName InPropertyName) const;
 
 	UFunction* ResolveSetter(FName InPropertyName, const UObject* InOwner);
 
 	/** Register a converter class */
-	bool RegisterConverterClass(const UClass* InConverterClass);
+	PROPERTYANIMATORCORE_API bool RegisterConverterClass(const UClass* InConverterClass);
 
 	/** Unregister a converter class */
-	bool UnregisterConverterClass(const UClass* InConverterClass);
+	PROPERTYANIMATORCORE_API bool UnregisterConverterClass(const UClass* InConverterClass);
 
 	/** Is this converter class registered */
-	bool IsConverterClassRegistered(const UClass* InConverterClass);
+	PROPERTYANIMATORCORE_API bool IsConverterClassRegistered(const UClass* InConverterClass);
 
 	/** Checks if any converter supports the type conversion */
-	bool IsConversionSupported(const FPropertyBagPropertyDesc& InFromProperty, const FPropertyBagPropertyDesc& InToProperty);
+	PROPERTYANIMATORCORE_API bool IsConversionSupported(const FPropertyBagPropertyDesc& InFromProperty, const FPropertyBagPropertyDesc& InToProperty);
 
 	/** Finds suitable converters for a type conversion */
-	TSet<UPropertyAnimatorCoreConverterBase*> GetSupportedConverters(const FPropertyBagPropertyDesc& InFromProperty, const FPropertyBagPropertyDesc& InToProperty) const;
+	PROPERTYANIMATORCORE_API TSet<UPropertyAnimatorCoreConverterBase*> GetSupportedConverters(const FPropertyBagPropertyDesc& InFromProperty, const FPropertyBagPropertyDesc& InToProperty) const;
 
 	/** Create an animator of specific class for an actor */
-	UPropertyAnimatorCoreBase* CreateAnimator(AActor* InActor, const UClass* InAnimatorClass, UPropertyAnimatorCorePresetBase* InPreset = nullptr, bool bInTransact = false) const;
+	PROPERTYANIMATORCORE_API UPropertyAnimatorCoreBase* CreateAnimator(AActor* InActor, const UClass* InAnimatorClass, UPropertyAnimatorCorePresetBase* InPreset = nullptr, bool bInTransact = false) const;
 
 	/** Create animators of specific class for actors */
-	TSet<UPropertyAnimatorCoreBase*> CreateAnimators(const TSet<AActor*>& InActors, const UClass* InAnimatorClass, UPropertyAnimatorCorePresetBase* InPreset = nullptr, bool bInTransact = false) const;
+	PROPERTYANIMATORCORE_API TSet<UPropertyAnimatorCoreBase*> CreateAnimators(const TSet<AActor*>& InActors, const UClass* InAnimatorClass, UPropertyAnimatorCorePresetBase* InPreset = nullptr, bool bInTransact = false) const;
 
 	/** Removes a animator bound to an owner */
-	bool RemoveAnimator(UPropertyAnimatorCoreBase* InAnimator, bool bInTransact = false) const;
+	PROPERTYANIMATORCORE_API bool RemoveAnimator(UPropertyAnimatorCoreBase* InAnimator, bool bInTransact = false) const;
 
 	/** Removes animators from their owner */
-	bool RemoveAnimators(const TSet<UPropertyAnimatorCoreBase*> InAnimators, bool bInTransact = false) const;
+	PROPERTYANIMATORCORE_API bool RemoveAnimators(const TSet<UPropertyAnimatorCoreBase*> InAnimators, bool bInTransact = false) const;
 
 	/** Apply a preset on an existing animator */
-	bool ApplyAnimatorPreset(UPropertyAnimatorCoreBase* InAnimator, UPropertyAnimatorCorePresetBase* InPreset, bool bInTransact = false);
+	PROPERTYANIMATORCORE_API bool ApplyAnimatorPreset(UPropertyAnimatorCoreBase* InAnimator, UPropertyAnimatorCorePresetBase* InPreset, bool bInTransact = false);
 
 	/** Unapply a preset from an existing animator */
-	bool UnapplyAnimatorPreset(UPropertyAnimatorCoreBase* InAnimator, UPropertyAnimatorCorePresetBase* InPreset, bool bInTransact = false);
+	PROPERTYANIMATORCORE_API bool UnapplyAnimatorPreset(UPropertyAnimatorCoreBase* InAnimator, UPropertyAnimatorCorePresetBase* InPreset, bool bInTransact = false);
 
 	/** Link a property to an existing animator */
-	bool LinkAnimatorProperty(UPropertyAnimatorCoreBase* InAnimator, FPropertyAnimatorCoreData& InProperty, bool bInTransact = false);
-	bool LinkAnimatorProperties(UPropertyAnimatorCoreBase* InAnimator, const TSet<FPropertyAnimatorCoreData>& InProperties, bool bInTransact = false);
+	PROPERTYANIMATORCORE_API bool LinkAnimatorProperty(UPropertyAnimatorCoreBase* InAnimator, FPropertyAnimatorCoreData& InProperty, bool bInTransact = false);
+	PROPERTYANIMATORCORE_API bool LinkAnimatorProperties(UPropertyAnimatorCoreBase* InAnimator, const TSet<FPropertyAnimatorCoreData>& InProperties, bool bInTransact = false);
 
 	/** Unlink a property from an existing animator */
-	bool UnlinkAnimatorProperty(UPropertyAnimatorCoreBase* InAnimator, FPropertyAnimatorCoreData& InProperty, bool bInTransact = false);
-	bool UnlinkAnimatorProperties(UPropertyAnimatorCoreBase* InAnimator, const TSet<FPropertyAnimatorCoreData>& InProperties, bool bInTransact = false);
+	PROPERTYANIMATORCORE_API bool UnlinkAnimatorProperty(UPropertyAnimatorCoreBase* InAnimator, FPropertyAnimatorCoreData& InProperty, bool bInTransact = false);
+	PROPERTYANIMATORCORE_API bool UnlinkAnimatorProperties(UPropertyAnimatorCoreBase* InAnimator, const TSet<FPropertyAnimatorCoreData>& InProperties, bool bInTransact = false);
 
 	/** Set the enabled state of animators attached to actors, will disable state globally on the component */
-	void SetActorAnimatorsEnabled(const TSet<AActor*>& InActors, bool bInEnabled, bool bInTransact = false);
+	PROPERTYANIMATORCORE_API void SetActorAnimatorsEnabled(const TSet<AActor*>& InActors, bool bInEnabled, bool bInTransact = false);
 
 	/** Set the enabled state of animators in a world, will disable state globally on the component */
-	void SetLevelAnimatorsEnabled(const UWorld* InWorld, bool bInEnabled, bool bInTransact = false);
+	PROPERTYANIMATORCORE_API void SetLevelAnimatorsEnabled(const UWorld* InWorld, bool bInEnabled, bool bInTransact = false);
 
 	/** Set the enabled state of animators provided */
-	void SetAnimatorsEnabled(const TSet<UPropertyAnimatorCoreBase*>& InAnimators, bool bInEnabled, bool bInTransact = false);
+	PROPERTYANIMATORCORE_API void SetAnimatorsEnabled(const TSet<UPropertyAnimatorCoreBase*>& InAnimators, bool bInEnabled, bool bInTransact = false);
 
 protected:
 	/** Delegate to change state of animators in a world */
diff --git a/Engine/Plugins/Experimental/PropertyAnimatorCore/Source/PropertyAnimatorCore/Public/TimeSources/PropertyAnimatorCoreManualTimeSource.h b/Engine/Plugins/Experimental/PropertyAnimatorCore/Source/PropertyAnimatorCore/Public/TimeSources/PropertyAnimatorCoreManualTimeSource.h
--- a/Engine/Plugins/Experimental/PropertyAnimatorCore/Source/PropertyAnimatorCore/Public/TimeSources/PropertyAnimatorCoreManualTimeSource.h
+++ b/Engine/Plugins/Experimental/PropertyAnimatorCore/Source/PropertyAnimatorCore/Public/TimeSources/PropertyAnimatorCoreManualTimeSource.h
@@ -5,8 +5,8 @@
 #include "PropertyAnimatorCoreTimeSourceBase.h"
 #include "PropertyAnimatorCoreManualTimeSource.generated.h"
 
-UCLASS()
-class PROPERTYANIMATORCORE_API UPropertyAnimatorCoreManualTimeSource : public UPropertyAnimatorCoreTimeSourceBase
+UCLASS(MinimalAPI)
+class UPropertyAnimatorCoreManualTimeSource : public UPropertyAnimatorCoreTimeSourceBase
 {
 	GENERATED_BODY()
 
@@ -15,9 +15,8 @@ public:
 		: UPropertyAnimatorCoreTimeSourceBase(TEXT("Manual"))
 	{}
 
-	UFUNCTION()
-	void SetCustomTime(float InTime);
-	float GetCustomTime() const
+	PROPERTYANIMATORCORE_API void SetCustomTime(double InTime);
+	double GetCustomTime() const
 	{
 		return CustomTime;
 	}
@@ -30,5 +29,5 @@ public:
 protected:
 	/** Allows you to drive controllers with this float */
 	UPROPERTY(EditInstanceOnly, BlueprintReadWrite, Setter="SetCustomTime", Getter="GetCustomTime", Category="Animator")
-	float CustomTime = 0.f;
+	double CustomTime = 0.f;
 };
\ No newline at end of file
diff --git a/Engine/Plugins/Experimental/PropertyAnimatorCore/Source/PropertyAnimatorCore/Public/TimeSources/PropertyAnimatorCoreRangeTimeSource.h b/Engine/Plugins/Experimental/PropertyAnimatorCore/Source/PropertyAnimatorCore/Public/TimeSources/PropertyAnimatorCoreRangeTimeSource.h
--- a/Engine/Plugins/Experimental/PropertyAnimatorCore/Source/PropertyAnimatorCore/Public/TimeSources/PropertyAnimatorCoreRangeTimeSource.h
+++ b/Engine/Plugins/Experimental/PropertyAnimatorCore/Source/PropertyAnimatorCore/Public/TimeSources/PropertyAnimatorCoreRangeTimeSource.h
@@ -8,8 +8,8 @@
 /**
  * Abstract range time source allowing to start or stop based on time elapsed
  */
-UCLASS(Abstract)
-class PROPERTYANIMATORCORE_API UPropertyAnimatorCoreRangeTimeSource : public UPropertyAnimatorCoreTimeSourceBase
+UCLASS(MinimalAPI, Abstract)
+class UPropertyAnimatorCoreRangeTimeSource : public UPropertyAnimatorCoreTimeSourceBase
 {
 	GENERATED_BODY()
 
@@ -22,39 +22,39 @@ public:
 		: UPropertyAnimatorCoreTimeSourceBase(InSourceName)
 	{}
 
-	void SetUseStartTime(bool bInUse);
+	PROPERTYANIMATORCORE_API void SetUseStartTime(bool bInUse);
 
 	bool GetUseStartTime() const
 	{
 		return bUseStartTime;
 	}
 
-	void SetStartTime(double InStartTime);
+	PROPERTYANIMATORCORE_API void SetStartTime(double InStartTime);
 
 	double GetStartTime() const
 	{
 		return StartTime;
 	}
 
-	void SetUseStopTime(bool bInUse);
+	PROPERTYANIMATORCORE_API void SetUseStopTime(bool bInUse);
 
 	bool GetUseStopTime() const
 	{
 		return bUseStopTime;
 	}
 
-	void SetStopTime(double InStopTime);
+	PROPERTYANIMATORCORE_API void SetStopTime(double InStopTime);
 
 	double GetStopTime() const
 	{
 		return StopTime;
 	}
 
+protected:
 	//~ Begin UPropertyAnimatorTimeSourceBase
-	virtual bool IsValidTimeElapsed(double InTimeElapsed) const override;
+	PROPERTYANIMATORCORE_API virtual bool IsValidTimeElapsed(double InTimeElapsed) const override;
 	//~ End UPropertyAnimatorTimeSourceBase
 
-protected:
 	UPROPERTY(EditInstanceOnly, BlueprintReadWrite, Setter="SetUseStartTime", Getter="GetUseStartTime", Category="Animator")
 	bool bUseStartTime = false;
 
diff --git a/Engine/Plugins/Experimental/PropertyAnimatorCore/Source/PropertyAnimatorCore/Public/TimeSources/PropertyAnimatorCoreTimeSourceBase.h b/Engine/Plugins/Experimental/PropertyAnimatorCore/Source/PropertyAnimatorCore/Public/TimeSources/PropertyAnimatorCoreTimeSourceBase.h
--- a/Engine/Plugins/Experimental/PropertyAnimatorCore/Source/PropertyAnimatorCore/Public/TimeSources/PropertyAnimatorCoreTimeSourceBase.h
+++ b/Engine/Plugins/Experimental/PropertyAnimatorCore/Source/PropertyAnimatorCore/Public/TimeSources/PropertyAnimatorCoreTimeSourceBase.h
@@ -10,8 +10,8 @@ class UPropertyAnimatorCoreBase;
  * Abstract base class for time source used by property animators
  * Can be transient or saved to disk if contains user set data
  */
-UCLASS(Abstract)
-class PROPERTYANIMATORCORE_API UPropertyAnimatorCoreTimeSourceBase : public UObject
+UCLASS(MinimalAPI, Abstract)
+class UPropertyAnimatorCoreTimeSourceBase : public UObject
 {
 	GENERATED_BODY()
 
diff --git a/Engine/Plugins/Experimental/PropertyAnimatorCore/Source/PropertyAnimatorCore/Public/TimeSources/PropertyAnimatorCoreWorldTimeSource.h b/Engine/Plugins/Experimental/PropertyAnimatorCore/Source/PropertyAnimatorCore/Public/TimeSources/PropertyAnimatorCoreWorldTimeSource.h
--- a/Engine/Plugins/Experimental/PropertyAnimatorCore/Source/PropertyAnimatorCore/Public/TimeSources/PropertyAnimatorCoreWorldTimeSource.h
+++ b/Engine/Plugins/Experimental/PropertyAnimatorCore/Source/PropertyAnimatorCore/Public/TimeSources/PropertyAnimatorCoreWorldTimeSource.h
@@ -9,8 +9,8 @@
  * Time source that follows world time,
  * Transient because it does not contain user saved data, can be recreated
  */
-UCLASS()
-class PROPERTYANIMATORCORE_API UPropertyAnimatorCoreWorldTimeSource : public UPropertyAnimatorCoreTimeSourceBase
+UCLASS(MinimalAPI, Transient)
+class UPropertyAnimatorCoreWorldTimeSource : public UPropertyAnimatorCoreTimeSourceBase
 {
 	GENERATED_BODY()
 
diff --git a/Engine/Plugins/Experimental/PropertyAnimatorCore/Source/PropertyAnimatorCoreEditor/Public/Subsystems/PropertyAnimatorCoreEditorSubsystem.h b/Engine/Plugins/Experimental/PropertyAnimatorCore/Source/PropertyAnimatorCoreEditor/Public/Subsystems/PropertyAnimatorCoreEditorSubsystem.h
--- a/Engine/Plugins/Experimental/PropertyAnimatorCore/Source/PropertyAnimatorCoreEditor/Public/Subsystems/PropertyAnimatorCoreEditorSubsystem.h
+++ b/Engine/Plugins/Experimental/PropertyAnimatorCore/Source/PropertyAnimatorCoreEditor/Public/Subsystems/PropertyAnimatorCoreEditorSubsystem.h
@@ -22,13 +22,13 @@ struct FPropertyRowExtensionButton;
 
 /** Singleton class that handles editor operations for property control such as windows */
 UCLASS()
-class PROPERTYANIMATORCOREEDITOR_API UPropertyAnimatorCoreEditorSubsystem : public UEditorSubsystem
+class UPropertyAnimatorCoreEditorSubsystem : public UEditorSubsystem
 {
 	GENERATED_BODY()
 
 public:
 	/** Get this subsystem instance */
-	static UPropertyAnimatorCoreEditorSubsystem* Get();
+	PROPERTYANIMATORCOREEDITOR_API static UPropertyAnimatorCoreEditorSubsystem* Get();
 
 	//~ Begin UEditorSubsystem
 	virtual void Initialize(FSubsystemCollectionBase& Collection) override;
@@ -45,7 +45,7 @@ public:
 	bool IsPropertyControlWindowOpened() const;
 
 	/** Fills a menu based on context objects and menu options */
-	bool FillAnimatorMenu(UToolMenu* InMenu, const FPropertyAnimatorCoreEditorMenuContext& InContext, const FPropertyAnimatorCoreEditorMenuOptions& InOptions);
+	PROPERTYANIMATORCOREEDITOR_API bool FillAnimatorMenu(UToolMenu* InMenu, const FPropertyAnimatorCoreEditorMenuContext& InContext, const FPropertyAnimatorCoreEditorMenuOptions& InOptions);
 
 protected:
 	/** Setup details panel button customization */
diff --git a/Engine/Plugins/Experimental/PropertyAnimatorCore/Source/PropertyAnimatorCoreEditor/Public/Widgets/PropertyAnimatorCoreEditorEditPanelOptions.h b/Engine/Plugins/Experimental/PropertyAnimatorCore/Source/PropertyAnimatorCoreEditor/Public/Widgets/PropertyAnimatorCoreEditorEditPanelOptions.h
--- a/Engine/Plugins/Experimental/PropertyAnimatorCore/Source/PropertyAnimatorCoreEditor/Public/Widgets/PropertyAnimatorCoreEditorEditPanelOptions.h
+++ b/Engine/Plugins/Experimental/PropertyAnimatorCore/Source/PropertyAnimatorCoreEditor/Public/Widgets/PropertyAnimatorCoreEditorEditPanelOptions.h
@@ -8,7 +8,7 @@
 class AActor;
 class SPropertyAnimatorCoreEditorEditPanel;
 
-struct PROPERTYANIMATORCOREEDITOR_API FPropertyAnimatorCoreEditorEditPanelOptions
+struct FPropertyAnimatorCoreEditorEditPanelOptions
 {
 	FPropertyAnimatorCoreEditorEditPanelOptions()
 	{}
