MotionDesign : PropertyAnimator
- Replace animated property owner (blueprint)
- Default mode set to additive instead of absolute
- Hiding ConverterRule when not needed

#changelist validated
#virtualized

diff --git a/Engine/Plugins/Experimental/PropertyAnimator/Source/PropertyAnimator/Private/Properties/PropertyAnimatorTextResolver.cpp b/Engine/Plugins/Experimental/PropertyAnimator/Source/PropertyAnimator/Private/Properties/PropertyAnimatorTextResolver.cpp
--- a/Engine/Plugins/Experimental/PropertyAnimator/Source/PropertyAnimator/Private/Properties/PropertyAnimatorTextResolver.cpp
+++ b/Engine/Plugins/Experimental/PropertyAnimator/Source/PropertyAnimator/Private/Properties/PropertyAnimatorTextResolver.cpp
@@ -50,11 +50,7 @@ void UPropertyAnimatorTextResolver::ResolveProperties(const FPropertyAnimatorCor
 		return;
 	}
 
-	TArray<FProperty*> ChainProperties;
-	for (const TFieldPath<FProperty>& Property : InTemplateProperty.GetChainProperties())
-	{
-		ChainProperties.Add(Property.Get());
-	}
+	const TArray<FProperty*> ChainProperties = InTemplateProperty.GetChainProperties();
 
 	// Gather each character in the text
 	for (int32 ComponentIndex = 0; ComponentIndex < TextRootComponent->GetNumChildrenComponents(); ComponentIndex++)
diff --git a/Engine/Plugins/Experimental/PropertyAnimatorCore/Source/PropertyAnimatorCore/Private/Animators/PropertyAnimatorCoreBase.cpp b/Engine/Plugins/Experimental/PropertyAnimatorCore/Source/PropertyAnimatorCore/Private/Animators/PropertyAnimatorCoreBase.cpp
--- a/Engine/Plugins/Experimental/PropertyAnimatorCore/Source/PropertyAnimatorCore/Private/Animators/PropertyAnimatorCoreBase.cpp
+++ b/Engine/Plugins/Experimental/PropertyAnimatorCore/Source/PropertyAnimatorCore/Private/Animators/PropertyAnimatorCoreBase.cpp
@@ -26,6 +26,13 @@ UPropertyAnimatorCoreBase::UPropertyAnimatorCoreBase()
 		const TArray<FName> TimeSources = AnimatorSubsystem->GetTimeSourceNames();
 		SetTimeSourceName(!TimeSources.IsEmpty() ? TimeSources[0] : NAME_None);
 	}
+
+#if WITH_EDITOR
+	if (!IsTemplate())
+	{
+		FCoreUObjectDelegates::OnObjectsReplaced.AddUObject(this, &UPropertyAnimatorCoreBase::OnObjectReplaced);
+	}
+#endif
 }
 
 UPropertyAnimatorCoreComponent* UPropertyAnimatorCoreBase::GetAnimatorComponent() const
@@ -101,6 +108,15 @@ UPropertyAnimatorCoreContext* UPropertyAnimatorCoreBase::GetLinkedPropertyContex
 	return PropertyOptions ? *PropertyOptions : nullptr;
 }
 
+void UPropertyAnimatorCoreBase::BeginDestroy()
+{
+	Super::BeginDestroy();
+
+#if WITH_EDITOR
+	FCoreUObjectDelegates::OnObjectsReplaced.RemoveAll(this);
+#endif
+}
+
 void UPropertyAnimatorCoreBase::PostLoad()
 {
 	CleanTimeSources();
@@ -114,6 +130,14 @@ void UPropertyAnimatorCoreBase::PostLoad()
 	OnAnimatorEnabledChanged();
 }
 
+void UPropertyAnimatorCoreBase::PreDuplicate(FObjectDuplicationParameters& InDupParams)
+{
+	Super::PreDuplicate(InDupParams);
+
+	constexpr bool bForceReset = true;
+	RestoreProperties(bForceReset);
+}
+
 #if WITH_EDITOR
 void UPropertyAnimatorCoreBase::PreEditUndo()
 {
@@ -133,7 +157,7 @@ void UPropertyAnimatorCoreBase::PostEditUndo()
 
 void UPropertyAnimatorCoreBase::PostEditChangeProperty(FPropertyChangedEvent& PropertyChangedEvent)
 {
-	UObject::PostEditChangeProperty(PropertyChangedEvent);
+	Super::PostEditChangeProperty(PropertyChangedEvent);
 
 	const FName MemberName = PropertyChangedEvent.GetMemberPropertyName();
 
@@ -250,12 +274,7 @@ bool UPropertyAnimatorCoreBase::GetPropertiesSupported(const FPropertyAnimatorCo
 		return !OutSupportedProperties.IsEmpty();
 	};
 
-	TArray<FProperty*> ChainProperties;
-	ChainProperties.Reserve(InPropertyData.GetChainProperties().Num());
-	for (const TFieldPath<FProperty>& ChainProperty : InPropertyData.GetChainProperties())
-	{
-		ChainProperties.Add(ChainProperty.Get());
-	}
+	TArray<FProperty*> ChainProperties = InPropertyData.GetChainProperties();
 
 	return FindSupportedPropertiesRecursively(ChainProperties, Owner, OutProperties);
 }
@@ -383,6 +402,25 @@ void UPropertyAnimatorCoreBase::EvaluateAnimator()
 	bEvaluatingProperties = false;
 }
 
+void UPropertyAnimatorCoreBase::OnObjectReplaced(const TMap<UObject*, UObject*>& InReplacementMap)
+{
+	constexpr bool bResolve = false;
+	ForEachLinkedProperty<UPropertyAnimatorCoreContext>([&InReplacementMap](UPropertyAnimatorCoreContext* InContext, const FPropertyAnimatorCoreData& InProperty)->bool
+	{
+		const TWeakObjectPtr<UObject> OwnerWeak = InProperty.GetOwnerWeak();
+
+		constexpr bool bEvenIfPendingKill = true;
+		const UObject* Owner = OwnerWeak.Get(bEvenIfPendingKill);
+
+		if (UObject* const* NewOwner = InReplacementMap.Find(Owner))
+		{
+			InContext->SetAnimatedPropertyOwner(*NewOwner);
+		}
+
+		return true;
+	}, bResolve);
+}
+
 void UPropertyAnimatorCoreBase::OnPropertyGroupsChanged()
 {
 	TSet<FName> CurrentGroupNames;
@@ -638,7 +676,7 @@ TSet<FPropertyAnimatorCoreData> UPropertyAnimatorCoreBase::GetInnerPropertiesLin
 		return OutProperties;
 	}
 
-	const TFieldPath<FProperty> LeafProperty = InPropertyData.GetLeafProperty();
+	FProperty* LeafProperty = InPropertyData.GetLeafProperty();
 
 	for (const FPropertyAnimatorCoreData& ControllerProperty : GetLinkedProperties())
 	{
diff --git a/Engine/Plugins/Experimental/PropertyAnimatorCore/Source/PropertyAnimatorCore/Private/Properties/PropertyAnimatorCoreContext.cpp b/Engine/Plugins/Experimental/PropertyAnimatorCore/Source/PropertyAnimatorCore/Private/Properties/PropertyAnimatorCoreContext.cpp
--- a/Engine/Plugins/Experimental/PropertyAnimatorCore/Source/PropertyAnimatorCore/Private/Properties/PropertyAnimatorCoreContext.cpp
+++ b/Engine/Plugins/Experimental/PropertyAnimatorCore/Source/PropertyAnimatorCore/Private/Properties/PropertyAnimatorCoreContext.cpp
@@ -105,6 +105,7 @@ void UPropertyAnimatorCoreContext::SetConverterClass(TSubclassOf<UPropertyAnimat
 		if (UScriptStruct* RuleStruct = Converter->GetConversionRuleStruct())
 		{
 			ConverterRule = FInstancedStruct(RuleStruct);
+			CheckEditConverterRule();
 		}
 	}
 }
@@ -130,7 +131,8 @@ void UPropertyAnimatorCoreContext::PostLoad()
 {
 	Super::PostLoad();
 
-	CheckModeEdit();
+	CheckEditMode();
+	CheckEditConverterRule();
 }
 
 #if WITH_EDITOR
@@ -201,18 +203,44 @@ void UPropertyAnimatorCoreContext::OnGroupNameChanged()
 void UPropertyAnimatorCoreContext::ConstructInternal(const FPropertyAnimatorCoreData& InProperty)
 {
 	AnimatedProperty = InProperty;
-	CheckModeEdit();
+	CheckEditMode();
+	CheckEditConverterRule();
+	SetMode(EPropertyAnimatorCoreMode::Additive);
 	OnAnimatedPropertyLinked();
 }
 
-void UPropertyAnimatorCoreContext::CheckModeEdit()
+void UPropertyAnimatorCoreContext::SetAnimatedPropertyOwner(UObject* InNewOwner)
+{
+	if (!IsValid(InNewOwner))
+	{
+		return;
+	}
+
+	if (!FindFProperty<FProperty>(InNewOwner->GetClass(), AnimatedProperty.GetMemberPropertyName()))
+	{
+		return;
+	}
+
+	constexpr bool bEvenIfPendingKill = true;
+	UObject* PreviousOwner = AnimatedProperty.GetOwnerWeak().Get(bEvenIfPendingKill);
+	AnimatedProperty = FPropertyAnimatorCoreData(InNewOwner, AnimatedProperty.GetChainProperties(), AnimatedProperty.GetPropertyResolverClass());
+
+	OnAnimatedPropertyOwnerUpdated(PreviousOwner, InNewOwner);
+}
+
+void UPropertyAnimatorCoreContext::CheckEditMode()
 {
 	if (const UPropertyAnimatorCoreHandlerBase* Handler = GetHandler())
 	{
-		bAllowModeEdit = Handler->IsAdditiveSupported();
+		bEditMode = Handler->IsAdditiveSupported();
 	}
 }
 
+void UPropertyAnimatorCoreContext::CheckEditConverterRule()
+{
+	bEditConverterRule = ConverterRule.IsValid();
+}
+
 void* UPropertyAnimatorCoreContext::GetConverterRulePtr(const UScriptStruct* InStruct)
 {
 	if (ConverterRule.IsValid() && ConverterRule.GetScriptStruct()->IsChildOf(InStruct))
diff --git a/Engine/Plugins/Experimental/PropertyAnimatorCore/Source/PropertyAnimatorCore/Private/Properties/PropertyAnimatorCoreData.cpp b/Engine/Plugins/Experimental/PropertyAnimatorCore/Source/PropertyAnimatorCore/Private/Properties/PropertyAnimatorCoreData.cpp
--- a/Engine/Plugins/Experimental/PropertyAnimatorCore/Source/PropertyAnimatorCore/Private/Properties/PropertyAnimatorCoreData.cpp
+++ b/Engine/Plugins/Experimental/PropertyAnimatorCore/Source/PropertyAnimatorCore/Private/Properties/PropertyAnimatorCoreData.cpp
@@ -135,6 +135,22 @@ FName FPropertyAnimatorCoreData::GetLeafPropertyName() const
 	return LeafProperty ? LeafProperty->GetFName() : NAME_None;
 }
 
+TArray<FProperty*> FPropertyAnimatorCoreData::GetChainProperties() const
+{
+	TArray<FProperty*> Properties;
+
+	Algo::Transform(
+		ChainProperties
+		, Properties
+		, [](const TFieldPath<FProperty>& InProperty)
+		{
+			return InProperty.Get();
+		}
+	);
+
+	return Properties;
+}
+
 bool FPropertyAnimatorCoreData::HasSetter() const
 {
 	FPropertyAnimatorCoreData* This = const_cast<FPropertyAnimatorCoreData*>(this);
@@ -145,7 +161,7 @@ bool FPropertyAnimatorCoreData::IsParentOf(const FPropertyAnimatorCoreData& InOt
 {
 	if (FProperty* LeafProperty = GetLeafProperty())
 	{
-		const TArray<TFieldPath<FProperty>>& OtherChainProperties = InOtherProperty.GetChainProperties();
+		const TArray<FProperty*> OtherChainProperties = InOtherProperty.GetChainProperties();
 
 		const int32 LeafIdx = OtherChainProperties.Find(LeafProperty);
 
@@ -165,7 +181,7 @@ bool FPropertyAnimatorCoreData::IsOwning(const FPropertyAnimatorCoreData& InOthe
 {
 	if (FProperty* LeafProperty = GetLeafProperty())
 	{
-		const TArray<TFieldPath<FProperty>>& OtherChainProperties = InOtherProperty.GetChainProperties();
+		const TArray<FProperty*> OtherChainProperties = InOtherProperty.GetChainProperties();
 
 		const int32 LeafIdx = OtherChainProperties.Find(LeafProperty);
 
diff --git a/Engine/Plugins/Experimental/PropertyAnimatorCore/Source/PropertyAnimatorCore/Public/Animators/PropertyAnimatorCoreBase.h b/Engine/Plugins/Experimental/PropertyAnimatorCore/Source/PropertyAnimatorCore/Public/Animators/PropertyAnimatorCoreBase.h
--- a/Engine/Plugins/Experimental/PropertyAnimatorCore/Source/PropertyAnimatorCore/Public/Animators/PropertyAnimatorCoreBase.h
+++ b/Engine/Plugins/Experimental/PropertyAnimatorCore/Source/PropertyAnimatorCore/Public/Animators/PropertyAnimatorCoreBase.h
@@ -147,7 +147,9 @@ public:
 
 protected:
 	//~ Begin UObject
+	PROPERTYANIMATORCORE_API virtual void BeginDestroy() override;
 	PROPERTYANIMATORCORE_API virtual void PostLoad() override;
+	PROPERTYANIMATORCORE_API virtual void PreDuplicate(FObjectDuplicationParameters& InDupParams) override;
 #if WITH_EDITOR
 	PROPERTYANIMATORCORE_API virtual void PreEditUndo() override;
 	PROPERTYANIMATORCORE_API virtual void PostEditUndo() override;
@@ -259,6 +261,8 @@ private:
 	/** Called by the component to evaluate this animator */
 	void EvaluateAnimator();
 
+	void OnObjectReplaced(const TMap<UObject*, UObject*>& InReplacementMap);
+
 	void OnPropertyGroupsChanged();
 
 	void OnAnimatorEnabledChanged();
diff --git a/Engine/Plugins/Experimental/PropertyAnimatorCore/Source/PropertyAnimatorCore/Public/Properties/PropertyAnimatorCoreContext.h b/Engine/Plugins/Experimental/PropertyAnimatorCore/Source/PropertyAnimatorCore/Public/Properties/PropertyAnimatorCoreContext.h
--- a/Engine/Plugins/Experimental/PropertyAnimatorCore/Source/PropertyAnimatorCore/Public/Properties/PropertyAnimatorCoreContext.h
+++ b/Engine/Plugins/Experimental/PropertyAnimatorCore/Source/PropertyAnimatorCore/Public/Properties/PropertyAnimatorCoreContext.h
@@ -91,12 +91,18 @@ protected:
 	/** Called once, when the property is linked to this context */
 	virtual void OnAnimatedPropertyLinked() {}
 
+	/** Called when the animated property owner is updated */
+	virtual void OnAnimatedPropertyOwnerUpdated(UObject* InPreviousOwner, UObject* InNewOwner) {}
+
 private:
 	void ConstructInternal(const FPropertyAnimatorCoreData& InProperty);
+	void SetAnimatedPropertyOwner(UObject* InNewOwner);
 
 	PROPERTYANIMATORCORE_API void* GetConverterRulePtr(const UScriptStruct* InStruct);
 
-	void CheckModeEdit();
+	void CheckEditMode();
+	void CheckEditConverterRule();
+
 	void OnAnimatedChanged();
 	void OnModeChanged();
 	void OnGroupNameChanged();
@@ -126,16 +132,20 @@ private:
 	UPROPERTY(EditInstanceOnly, Setter="SetAnimated", Getter="IsAnimated", Category="Animator", meta=(AllowPrivateAccess="true"))
 	bool bAnimated = true;
 
-	/** Edit condition for modes (Force hide otherwise appears inline for mode) */
-	UPROPERTY(Transient, VisibleInstanceOnly, Category="Animator", meta=(EditCondition="false", EditConditionHides))
-	bool bAllowModeEdit = true;
+	/** Edit condition for modes */
+	UPROPERTY(Transient)
+	bool bEditMode = true;
 
 	/** Current mode used for this property */
-	UPROPERTY(EditInstanceOnly, Setter, Getter, Category="Animator", meta=(EditCondition="bAllowModeEdit", EditConditionHides, AllowPrivateAccess="true"))
+	UPROPERTY(EditInstanceOnly, Setter, Getter, Category="Animator", meta=(HideEditConditionToggle, EditCondition="bEditMode", EditConditionHides, AllowPrivateAccess="true"))
 	EPropertyAnimatorCoreMode Mode = EPropertyAnimatorCoreMode::Absolute;
 
+	/** Edit condition for converter rule */
+	UPROPERTY(Transient)
+	bool bEditConverterRule = false;
+
 	/** If a converter is used, rules may be used to convert the property */
-	UPROPERTY(EditInstanceOnly, Category="Animator", meta=(AllowPrivateAccess="true"))
+	UPROPERTY(EditInstanceOnly, Category="Animator", meta=(HideEditConditionToggle, EditCondition="bEditConverterRule", EditConditionHides, AllowPrivateAccess="true"))
 	FInstancedStruct ConverterRule;
 
 	/** The unique group name that manages this property */
@@ -147,11 +157,11 @@ private:
 	TObjectPtr<UPropertyAnimatorCoreGroupBase> Group;
 
 	/** Store original property values for resolved properties */
-	UPROPERTY(DuplicateTransient, NonTransactional)
+	UPROPERTY(NonTransactional)
 	FInstancedPropertyBag OriginalPropertyValues;
 
 	/** Store delta property values for resolved properties */
-	UPROPERTY(DuplicateTransient, NonTransactional)
+	UPROPERTY(NonTransactional)
 	FInstancedPropertyBag DeltaPropertyValues;
 
 	/** Converter class used for this property */
diff --git a/Engine/Plugins/Experimental/PropertyAnimatorCore/Source/PropertyAnimatorCore/Public/Properties/PropertyAnimatorCoreData.h b/Engine/Plugins/Experimental/PropertyAnimatorCore/Source/PropertyAnimatorCore/Public/Properties/PropertyAnimatorCoreData.h
--- a/Engine/Plugins/Experimental/PropertyAnimatorCore/Source/PropertyAnimatorCore/Public/Properties/PropertyAnimatorCoreData.h
+++ b/Engine/Plugins/Experimental/PropertyAnimatorCore/Source/PropertyAnimatorCore/Public/Properties/PropertyAnimatorCoreData.h
@@ -72,6 +72,11 @@ struct FPropertyAnimatorCoreData
 		return OwnerWeak.Get();
 	}
 
+	TWeakObjectPtr<UObject> GetOwnerWeak() const
+	{
+		return OwnerWeak;
+	}
+
 	/** The member property of the owner, top property inside the owner itself */
 	FProperty* GetMemberProperty() const
 	{
@@ -103,10 +108,7 @@ struct FPropertyAnimatorCoreData
 	PROPERTYANIMATORCORE_API FName GetLeafPropertyName() const;
 
 	/** The chain properties from member to inner property */
-	const TArray<TFieldPath<FProperty>>& GetChainProperties() const
-	{
-		return ChainProperties;
-	}
+	PROPERTYANIMATORCORE_API TArray<FProperty*> GetChainProperties() const;
 
 	/** Checks if the property is settable via setter */
 	PROPERTYANIMATORCORE_API bool HasSetter() const;
