MotionDesign : PropertyAnimatorCore
- Using array instead of set to store animators in component

diff --git a/Engine/Plugins/Experimental/PropertyAnimatorCore/Source/PropertyAnimatorCore/Private/Components/PropertyAnimatorCoreComponent.cpp b/Engine/Plugins/Experimental/PropertyAnimatorCore/Source/PropertyAnimatorCore/Private/Components/PropertyAnimatorCoreComponent.cpp
--- a/Engine/Plugins/Experimental/PropertyAnimatorCore/Source/PropertyAnimatorCore/Private/Components/PropertyAnimatorCoreComponent.cpp
+++ b/Engine/Plugins/Experimental/PropertyAnimatorCore/Source/PropertyAnimatorCore/Private/Components/PropertyAnimatorCoreComponent.cpp
@@ -17,8 +17,8 @@ UPropertyAnimatorCoreBase* UPropertyAnimatorCoreComponent::AddAnimator(const UCl
 
 	if (NewAnimator)
 	{
-		AnimatorsInternal = Animators;
-		Animators.Add(NewAnimator);
+		PropertyAnimatorsInternal = PropertyAnimators;
+		PropertyAnimators.Add(NewAnimator);
 
 		OnAnimatorsChanged();
 	}
@@ -42,8 +42,8 @@ UPropertyAnimatorCoreBase* UPropertyAnimatorCoreComponent::CloneAnimator(UProper
 	// Force current state
 	CloneAnimator->OnAnimatorEnabledChanged();
 
-	AnimatorsInternal = Animators;
-	Animators.Add(CloneAnimator);
+	PropertyAnimatorsInternal = PropertyAnimators;
+	PropertyAnimators.Add(CloneAnimator);
 
 	OnAnimatorsChanged();
 
@@ -52,13 +52,13 @@ UPropertyAnimatorCoreBase* UPropertyAnimatorCoreComponent::CloneAnimator(UProper
 
 bool UPropertyAnimatorCoreComponent::RemoveAnimator(UPropertyAnimatorCoreBase* InAnimator)
 {
-	if (!Animators.Contains(InAnimator))
+	if (!PropertyAnimators.Contains(InAnimator))
 	{
 		return false;
 	}
 
-	AnimatorsInternal = Animators;
-	Animators.Remove(InAnimator);
+	PropertyAnimatorsInternal = PropertyAnimators;
+	PropertyAnimators.Remove(InAnimator);
 
 	OnAnimatorsChanged();
 
@@ -82,9 +82,12 @@ void UPropertyAnimatorCoreComponent::OnAnimatorsSetEnabled(const UWorld* InWorld
 
 void UPropertyAnimatorCoreComponent::OnAnimatorsChanged()
 {
-	TSet<TObjectPtr<UPropertyAnimatorCoreBase>> RemovedAnimators = AnimatorsInternal.Difference(Animators);
-	TSet<TObjectPtr<UPropertyAnimatorCoreBase>> AddedAnimators = Animators.Difference(AnimatorsInternal);
-	AnimatorsInternal.Empty();
+	const TSet<TObjectPtr<UPropertyAnimatorCoreBase>> AnimatorsSet(PropertyAnimators);
+	const TSet<TObjectPtr<UPropertyAnimatorCoreBase>> AnimatorsInternalSet(PropertyAnimatorsInternal);
+
+	TSet<TObjectPtr<UPropertyAnimatorCoreBase>> RemovedAnimators = AnimatorsInternalSet.Difference(AnimatorsSet);
+	TSet<TObjectPtr<UPropertyAnimatorCoreBase>> AddedAnimators = AnimatorsSet.Difference(AnimatorsInternalSet);
+	PropertyAnimatorsInternal.Empty();
 
 	for (const TObjectPtr<UPropertyAnimatorCoreBase>& RemovedAnimator : RemovedAnimators)
 	{
@@ -123,7 +126,7 @@ void UPropertyAnimatorCoreComponent::OnAnimatorsEnabledChanged()
 		return;
 	}
 
-	for (const TObjectPtr<UPropertyAnimatorCoreBase>& Animator : Animators)
+	for (const TObjectPtr<UPropertyAnimatorCoreBase>& Animator : PropertyAnimators)
 	{
 		if (!IsValid(Animator))
 		{
@@ -146,7 +149,7 @@ void UPropertyAnimatorCoreComponent::OnAnimatorsEnabledChanged()
 bool UPropertyAnimatorCoreComponent::ShouldAnimatorsTick() const
 {
 	return bAnimatorsEnabled
-		&& !Animators.IsEmpty()
+		&& !PropertyAnimators.IsEmpty()
 		&& !FMath::IsNearlyZero(AnimatorsMagnitude);
 }
 
@@ -185,7 +188,7 @@ void UPropertyAnimatorCoreComponent::TickComponent(float InDeltaTime, ELevelTick
 		return;
 	}
 
-	for (const TObjectPtr<UPropertyAnimatorCoreBase>& Animator : Animators)
+	for (const TObjectPtr<UPropertyAnimatorCoreBase>& Animator : PropertyAnimators)
 	{
 		if (!IsValid(Animator))
 		{
@@ -272,14 +275,26 @@ void UPropertyAnimatorCoreComponent::SetAnimatorsMagnitude(float InMagnitude)
 
 void UPropertyAnimatorCoreComponent::DestroyComponent(bool bPromoteChildren)
 {
-	AnimatorsInternal = Animators;
-	Animators.Empty();
+	PropertyAnimatorsInternal = PropertyAnimators;
+	PropertyAnimators.Empty();
 
 	OnAnimatorsChanged();
 
 	Super::DestroyComponent(bPromoteChildren);
 }
 
+void UPropertyAnimatorCoreComponent::PostLoad()
+{
+	Super::PostLoad();
+
+	// Migrate animators to new array property
+	if (!Animators.IsEmpty() && PropertyAnimators.IsEmpty())
+	{
+		PropertyAnimators = Animators.Array();
+		Animators.Empty();
+	}
+}
+
 #if WITH_EDITOR
 void UPropertyAnimatorCoreComponent::PostEditUndo()
 {
@@ -294,9 +309,9 @@ void UPropertyAnimatorCoreComponent::PreEditChange(FProperty* PropertyAboutToCha
 
 	const FName MemberName = PropertyAboutToChange ? PropertyAboutToChange->GetFName() : NAME_None;
 
-	if (MemberName == GET_MEMBER_NAME_CHECKED(UPropertyAnimatorCoreComponent, Animators))
+	if (MemberName == GET_MEMBER_NAME_CHECKED(UPropertyAnimatorCoreComponent, PropertyAnimators))
 	{
-		AnimatorsInternal = Animators;
+		PropertyAnimatorsInternal = PropertyAnimators;
 	}
 }
 
@@ -311,28 +326,23 @@ void UPropertyAnimatorCoreComponent::PostEditChangeProperty(FPropertyChangedEven
 	{
 		OnAnimatorsEnabledChanged();
 	}
-	else if (MemberName == GET_MEMBER_NAME_CHECKED(UPropertyAnimatorCoreComponent, Animators))
+	else if (MemberName == GET_MEMBER_NAME_CHECKED(UPropertyAnimatorCoreComponent, PropertyAnimators))
 	{
 		OnAnimatorsChanged();
 	}
 }
 #endif
 
-void UPropertyAnimatorCoreComponent::SetAnimators(const TSet<TObjectPtr<UPropertyAnimatorCoreBase>>& InAnimators)
+void UPropertyAnimatorCoreComponent::SetAnimators(const TArray<TObjectPtr<UPropertyAnimatorCoreBase>>& InAnimators)
 {
-	if (Animators.Includes(InAnimators) && Animators.Num() == InAnimators.Num())
-	{
-		return;
-	}
-
-	AnimatorsInternal = Animators;
-	Animators = InAnimators;
+	PropertyAnimatorsInternal = PropertyAnimators;
+	PropertyAnimators = InAnimators;
 	OnAnimatorsChanged();
 }
 
 void UPropertyAnimatorCoreComponent::ForEachAnimator(const TFunctionRef<bool(UPropertyAnimatorCoreBase*)> InFunction) const
 {
-	for (const TObjectPtr<UPropertyAnimatorCoreBase>& Animator : Animators)
+	for (const TObjectPtr<UPropertyAnimatorCoreBase>& Animator : PropertyAnimators)
 	{
 		if (Animator)
 		{
diff --git a/Engine/Plugins/Experimental/PropertyAnimatorCore/Source/PropertyAnimatorCore/Public/Components/PropertyAnimatorCoreComponent.h b/Engine/Plugins/Experimental/PropertyAnimatorCore/Source/PropertyAnimatorCore/Public/Components/PropertyAnimatorCoreComponent.h
--- a/Engine/Plugins/Experimental/PropertyAnimatorCore/Source/PropertyAnimatorCore/Public/Components/PropertyAnimatorCoreComponent.h
+++ b/Engine/Plugins/Experimental/PropertyAnimatorCore/Source/PropertyAnimatorCore/Public/Components/PropertyAnimatorCoreComponent.h
@@ -21,15 +21,15 @@ public:
 
 	UPropertyAnimatorCoreComponent();
 
-	void SetAnimators(const TSet<TObjectPtr<UPropertyAnimatorCoreBase>>& InAnimators);
-	const TSet<TObjectPtr<UPropertyAnimatorCoreBase>>& GetAnimators() const
+	void SetAnimators(const TArray<TObjectPtr<UPropertyAnimatorCoreBase>>& InAnimators);
+	TConstArrayView<TObjectPtr<UPropertyAnimatorCoreBase>> GetAnimators() const
 	{
-		return Animators;
+		return PropertyAnimators;
 	}
 
 	int32 GetAnimatorsCount() const
 	{
-		return Animators.Num();
+		return PropertyAnimators.Num();
 	}
 
 	/** Set the state of all animators in this component */
@@ -57,6 +57,7 @@ protected:
 	//~ End UActorComponent
 
 	//~ Begin UObject
+	virtual void PostLoad() override;
 #if WITH_EDITOR
 	virtual void PostEditUndo() override;
 	virtual void PreEditChange(FProperty* PropertyAboutToChange) override;
@@ -92,8 +93,8 @@ protected:
 	bool ShouldAnimatorsTick() const;
 
 	/** Animators linked to this actor, they contain only properties within this actor */
-	UPROPERTY(EditAnywhere, BlueprintReadWrite, Instanced, Getter="GetAnimators", Setter="SetAnimators", Category="Animator", meta=(TitleProperty="AnimatorDisplayName"))
-	TSet<TObjectPtr<UPropertyAnimatorCoreBase>> Animators;
+	UPROPERTY(EditAnywhere, BlueprintReadWrite, Instanced, Setter="SetAnimators", Category="Animator", meta=(TitleProperty="AnimatorDisplayName"))
+	TArray<TObjectPtr<UPropertyAnimatorCoreBase>> PropertyAnimators;
 
 	/** Global state for all animators controlled by this component */
 	UPROPERTY(EditInstanceOnly, Getter="GetAnimatorsEnabled", Setter="SetAnimatorsEnabled", Category="Animator", meta=(DisplayPriority="0", AllowPrivateAccess="true"))
@@ -106,7 +107,11 @@ protected:
 private:
 	virtual void TickComponent(float InDeltaTime, ELevelTick InTickType, FActorComponentTickFunction* InThisTickFunction) override;
 
-	/** Transient copy of animators set when changes are detected to see the diff only */
+	/** Deprecated property set, will be migrated to PropertyAnimators property on load */
+	UPROPERTY()
+	TSet<TObjectPtr<UPropertyAnimatorCoreBase>> Animators;
+
+	/** Transient copy of property animators when changes are detected to see the diff only */
 	UPROPERTY(Transient, DuplicateTransient, TextExportTransient)
-	TSet<TObjectPtr<UPropertyAnimatorCoreBase>> AnimatorsInternal;
+	TArray<TObjectPtr<UPropertyAnimatorCoreBase>> PropertyAnimatorsInternal;
 };
\ No newline at end of file
