MotionDesign : PropertyAnimatorCore
- Fix issue with spawnable actors with animators via template actors (sequencer)

diff --git a/Engine/Plugins/Experimental/PropertyAnimatorCore/Source/PropertyAnimatorCore/Private/Animators/PropertyAnimatorCoreBase.cpp b/Engine/Plugins/Experimental/PropertyAnimatorCore/Source/PropertyAnimatorCore/Private/Animators/PropertyAnimatorCoreBase.cpp
--- a/Engine/Plugins/Experimental/PropertyAnimatorCore/Source/PropertyAnimatorCore/Private/Animators/PropertyAnimatorCoreBase.cpp
+++ b/Engine/Plugins/Experimental/PropertyAnimatorCore/Source/PropertyAnimatorCore/Private/Animators/PropertyAnimatorCoreBase.cpp
@@ -134,7 +134,7 @@ void UPropertyAnimatorCoreBase::PostEditImport()
 {
 	Super::PostEditImport();
 
-	ResolveProperties();
+	ResolvePropertiesOwner();
 }
 
 void UPropertyAnimatorCoreBase::PreDuplicate(FObjectDuplicationParameters& InParams)
@@ -149,7 +149,7 @@ void UPropertyAnimatorCoreBase::PostDuplicate(EDuplicateMode::Type InMode)
 {
 	Super::PostDuplicate(InMode);
 
-	ResolveProperties();
+	ResolvePropertiesOwner();
 }
 
 #if WITH_EDITOR
@@ -386,15 +386,15 @@ void UPropertyAnimatorCoreBase::OnTimeSourceNameChanged()
 	OnTimeSourceChanged();
 }
 
-void UPropertyAnimatorCoreBase::ResolveProperties()
+void UPropertyAnimatorCoreBase::ResolvePropertiesOwner(AActor* InNewOwner)
 {
 	// Resolve linked properties against current actor
 	TSet<FPropertyAnimatorCoreData> UnresolvedProperties;
 
 	ForEachLinkedProperty<UPropertyAnimatorCoreContext>(
-		[this, &UnresolvedProperties](UPropertyAnimatorCoreContext* InContext, const FPropertyAnimatorCoreData& InProperty)->bool
+		[this, &UnresolvedProperties, &InNewOwner](UPropertyAnimatorCoreContext* InContext, const FPropertyAnimatorCoreData& InProperty)->bool
 		{
-			if (!InContext->ResolveProperty())
+			if (!InContext->ResolvePropertyOwner(InNewOwner))
 			{
 				UnresolvedProperties.Add(InProperty);
 			}
diff --git a/Engine/Plugins/Experimental/PropertyAnimatorCore/Source/PropertyAnimatorCore/Private/Components/PropertyAnimatorCoreComponent.cpp b/Engine/Plugins/Experimental/PropertyAnimatorCore/Source/PropertyAnimatorCore/Private/Components/PropertyAnimatorCoreComponent.cpp
--- a/Engine/Plugins/Experimental/PropertyAnimatorCore/Source/PropertyAnimatorCore/Private/Components/PropertyAnimatorCoreComponent.cpp
+++ b/Engine/Plugins/Experimental/PropertyAnimatorCore/Source/PropertyAnimatorCore/Private/Components/PropertyAnimatorCoreComponent.cpp
@@ -171,6 +171,26 @@ FName UPropertyAnimatorCoreComponent::GetAnimatorName(const UPropertyAnimatorCor
 	return FName(NewAnimatorName);
 }
 
+void UPropertyAnimatorCoreComponent::OnComponentCreated()
+{
+	Super::OnComponentCreated();
+
+	if (AActor* OwningActor = GetOwner())
+	{
+		// For spawnable templates, restore and resolve properties owner
+		constexpr bool bForceRestore = true;
+
+		for (const TObjectPtr<UPropertyAnimatorCoreBase>& Animator : PropertyAnimators)
+		{
+			if (Animator)
+			{
+				Animator->RestoreProperties(bForceRestore);
+				Animator->ResolvePropertiesOwner(OwningActor);
+			}
+		}
+	}
+}
+
 void UPropertyAnimatorCoreComponent::TickComponent(float InDeltaTime, ELevelTick InTickType, FActorComponentTickFunction* InThisTickFunction)
 {
 	Super::TickComponent(InDeltaTime, InTickType, InThisTickFunction);
diff --git a/Engine/Plugins/Experimental/PropertyAnimatorCore/Source/PropertyAnimatorCore/Private/Properties/PropertyAnimatorCoreContext.cpp b/Engine/Plugins/Experimental/PropertyAnimatorCore/Source/PropertyAnimatorCore/Private/Properties/PropertyAnimatorCoreContext.cpp
--- a/Engine/Plugins/Experimental/PropertyAnimatorCore/Source/PropertyAnimatorCore/Private/Properties/PropertyAnimatorCoreContext.cpp
+++ b/Engine/Plugins/Experimental/PropertyAnimatorCore/Source/PropertyAnimatorCore/Private/Properties/PropertyAnimatorCoreContext.cpp
@@ -2,6 +2,7 @@
 
 #include "Properties/PropertyAnimatorCoreContext.h"
 
+#include "Containers/Ticker.h"
 #include "InstancedStruct.h"
 #include "Properties/PropertyAnimatorCoreGroupBase.h"
 #include "Properties/PropertyAnimatorCoreResolver.h"
@@ -133,6 +134,14 @@ void UPropertyAnimatorCoreContext::PostLoad()
 
 	CheckEditMode();
 	CheckEditConverterRule();
+
+	FTSTicker::GetCoreTicker().AddTicker(FTickerDelegate::CreateWeakLambda(this, [this](float InDelta)
+	{
+		// Restore before regenerating new property path
+		Restore();
+		AnimatedProperty.GeneratePropertyPath();
+		return false;
+	}));
 }
 
 #if WITH_EDITOR
@@ -200,37 +209,69 @@ void UPropertyAnimatorCoreContext::OnGroupNameChanged()
 	}
 }
 
-bool UPropertyAnimatorCoreContext::ResolveProperty()
+bool UPropertyAnimatorCoreContext::ResolvePropertyOwner(AActor* InNewOwner)
 {
-	UObject* NewOwner = GetTypedOuter<AActor>();
+	UObject* NewOwner = InNewOwner ? InNewOwner : GetTypedOuter<AActor>();
+	const UObject* CurrentOwner = AnimatedProperty.GetOwningActor();
 
-	if (AnimatedProperty.GetOwningActor() == NewOwner)
+	if (CurrentOwner == NewOwner)
 	{
 		return true;
 	}
 
 	bool bFound = IsValid(NewOwner);
+	const TArray<UObject*> OtherOuters = AnimatedProperty.GetOuters(AnimatedProperty.GetOwningActor());
 
-	for (const UObject* OtherOuter : AnimatedProperty.GetOuters(AnimatedProperty.GetOwningActor()))
+	if (!OtherOuters.IsEmpty())
 	{
-		bFound = false;
-		TArray<UObject*> ThisOwnedObjects;
-		GetObjectsWithOuter(NewOwner, ThisOwnedObjects, false);
+		// Resolve using outers
+		for (const UObject* OtherOuter : OtherOuters)
+		{
+			bFound = false;
+			TArray<UObject*> ThisOwnedObjects;
+			GetObjectsWithOuter(NewOwner, ThisOwnedObjects, false);
 
-		for (UObject* ThisOuter : ThisOwnedObjects)
+			for (UObject* ThisOuter : ThisOwnedObjects)
+			{
+				if (ThisOuter->GetClass() == OtherOuter->GetClass()
+					&& ThisOuter->GetFName() == OtherOuter->GetFName())
+				{
+					bFound = true;
+					NewOwner = ThisOuter;
+					break;
+				}
+			}
+		}
+	}
+	else
+	{
+		// Resolve using path segments
+		const TArray<FString> OwnerSegments = AnimatedProperty.GetOuterNames();
+
+		for (const FString& PathSegment : OwnerSegments)
 		{
-			if (ThisOuter->GetClass() == OtherOuter->GetClass()
-				&& ThisOuter->GetFName() == OtherOuter->GetFName())
+			bFound = false;
+			TArray<UObject*> ThisOwnedObjects;
+			GetObjectsWithOuter(NewOwner, ThisOwnedObjects, false);
+
+			for (UObject* ThisOuter : ThisOwnedObjects)
 			{
-				bFound = true;
-				NewOwner = ThisOuter;
-				break;
+				if (ThisOuter
+					&& ThisOuter->GetName().StartsWith(PathSegment))
+				{
+					bFound = true;
+					NewOwner = ThisOuter;
+					break;
+				}
 			}
 		}
 	}
 
+	const UObject* PreviousPropertyOwner = AnimatedProperty.GetOwner();
+
 	if (bFound
-		&& NewOwner->GetClass() == AnimatedProperty.GetOwner()->GetClass()
+		&& IsValid(NewOwner)
+		&& (!PreviousPropertyOwner || NewOwner->GetClass() == PreviousPropertyOwner->GetClass())
 		&& FindFProperty<FProperty>(NewOwner->GetClass(), AnimatedProperty.GetMemberPropertyName()))
 	{
 		SetAnimatedPropertyOwner(NewOwner);
diff --git a/Engine/Plugins/Experimental/PropertyAnimatorCore/Source/PropertyAnimatorCore/Private/Properties/PropertyAnimatorCoreData.cpp b/Engine/Plugins/Experimental/PropertyAnimatorCore/Source/PropertyAnimatorCore/Private/Properties/PropertyAnimatorCoreData.cpp
--- a/Engine/Plugins/Experimental/PropertyAnimatorCore/Source/PropertyAnimatorCore/Private/Properties/PropertyAnimatorCoreData.cpp
+++ b/Engine/Plugins/Experimental/PropertyAnimatorCore/Source/PropertyAnimatorCore/Private/Properties/PropertyAnimatorCoreData.cpp
@@ -145,6 +145,31 @@ TArray<UObject*> FPropertyAnimatorCoreData::GetOuters(const UObject* InStopOuter
 	return Owners;
 }
 
+TArray<FString> FPropertyAnimatorCoreData::GetOuterNames() const
+{
+	TArray<FString> OuterNames;
+	PathHash.ParseIntoArray(OuterNames, TEXT("."), true);
+	const FString MemberPropertyName = GetMemberPropertyName().ToString();
+
+	// Remove all properties chain segments
+	for (int32 Index = OuterNames.Num() - 1; Index >= 0; Index--)
+	{
+		if (OuterNames[Index].Equals(MemberPropertyName))
+		{
+			OuterNames.RemoveAt(Index, OuterNames.Num() - Index);
+			break;
+		}
+	}
+
+	// Remove resolver name at the beginning
+	if (!OuterNames.IsEmpty() && IsResolvable())
+	{
+		OuterNames.RemoveAt(0);
+	}
+
+	return OuterNames;
+}
+
 FName FPropertyAnimatorCoreData::GetMemberPropertyName() const
 {
 	const FProperty* MemberProperty = GetMemberProperty();
@@ -481,10 +506,11 @@ void FPropertyAnimatorCoreData::CopyPropertyValue(const FProperty* InProperty, c
 void FPropertyAnimatorCoreData::GeneratePropertyPath()
 {
 	const UObject* Owner = GetOwner();
-	const FString ResolverName = IsResolvable() ? GetPropertyResolver()->GetResolverName().ToString() : TEXT("");
+	const UObject* StopOuter = GetOwningActor();
+	const FString ResolverName = IsResolvable() ? GetPropertyResolver()->GetResolverName().ToString() + TEXT(".") : TEXT("");
 
 	PathHash = ResolverName;
-	PathHash += IsValid(Owner) ? Owner->GetPathName() : TEXT("");
+	PathHash += IsValid(Owner) ? Owner->GetPathName(StopOuter) : TEXT("");
 
 	FString DisplayName = ResolverName;
 	for (const TFieldPath<FProperty>& ChainProperty : ChainProperties)
diff --git a/Engine/Plugins/Experimental/PropertyAnimatorCore/Source/PropertyAnimatorCore/Public/Animators/PropertyAnimatorCoreBase.h b/Engine/Plugins/Experimental/PropertyAnimatorCore/Source/PropertyAnimatorCore/Public/Animators/PropertyAnimatorCoreBase.h
--- a/Engine/Plugins/Experimental/PropertyAnimatorCore/Source/PropertyAnimatorCore/Public/Animators/PropertyAnimatorCoreBase.h
+++ b/Engine/Plugins/Experimental/PropertyAnimatorCore/Source/PropertyAnimatorCore/Public/Animators/PropertyAnimatorCoreBase.h
@@ -279,7 +279,7 @@ private:
 	void OnTimeSourceNameChanged();
 
 	/** Called after an action that causes the owner to change */
-	void ResolveProperties();
+	void ResolvePropertiesOwner(AActor* InNewOwner = nullptr);
 
 	UPropertyAnimatorCoreTimeSourceBase* FindOrAddTimeSource(FName InTimeSourceName);
 
diff --git a/Engine/Plugins/Experimental/PropertyAnimatorCore/Source/PropertyAnimatorCore/Public/Components/PropertyAnimatorCoreComponent.h b/Engine/Plugins/Experimental/PropertyAnimatorCore/Source/PropertyAnimatorCore/Public/Components/PropertyAnimatorCoreComponent.h
--- a/Engine/Plugins/Experimental/PropertyAnimatorCore/Source/PropertyAnimatorCore/Public/Components/PropertyAnimatorCoreComponent.h
+++ b/Engine/Plugins/Experimental/PropertyAnimatorCore/Source/PropertyAnimatorCore/Public/Components/PropertyAnimatorCoreComponent.h
@@ -53,6 +53,7 @@ protected:
 	static FName GetAnimatorName(const UPropertyAnimatorCoreBase* InAnimator);
 
 	//~ Begin UActorComponent
+	virtual void OnComponentCreated() override;
 	virtual void DestroyComponent(bool bPromoteChildren) override;
 	//~ End UActorComponent
 
diff --git a/Engine/Plugins/Experimental/PropertyAnimatorCore/Source/PropertyAnimatorCore/Public/Properties/PropertyAnimatorCoreContext.h b/Engine/Plugins/Experimental/PropertyAnimatorCore/Source/PropertyAnimatorCore/Public/Properties/PropertyAnimatorCoreContext.h
--- a/Engine/Plugins/Experimental/PropertyAnimatorCore/Source/PropertyAnimatorCore/Public/Properties/PropertyAnimatorCoreContext.h
+++ b/Engine/Plugins/Experimental/PropertyAnimatorCore/Source/PropertyAnimatorCore/Public/Properties/PropertyAnimatorCoreContext.h
@@ -80,7 +80,7 @@ public:
 	}
 
 	/** Called when the owner has changed and we want to update the animated property */
-	bool ResolveProperty();
+	bool ResolvePropertyOwner(AActor* InNewOwner);
 
 protected:
 	//~ Begin UObject
diff --git a/Engine/Plugins/Experimental/PropertyAnimatorCore/Source/PropertyAnimatorCore/Public/Properties/PropertyAnimatorCoreData.h b/Engine/Plugins/Experimental/PropertyAnimatorCore/Source/PropertyAnimatorCore/Public/Properties/PropertyAnimatorCoreData.h
--- a/Engine/Plugins/Experimental/PropertyAnimatorCore/Source/PropertyAnimatorCore/Public/Properties/PropertyAnimatorCoreData.h
+++ b/Engine/Plugins/Experimental/PropertyAnimatorCore/Source/PropertyAnimatorCore/Public/Properties/PropertyAnimatorCoreData.h
@@ -80,6 +80,8 @@ struct FPropertyAnimatorCoreData
 	/** Returns chain of owner until it reaches StopOuter */
 	TArray<UObject*> GetOuters(const UObject* InStopOuter) const;
 
+	TArray<FString> GetOuterNames() const;
+
 	/** The member property of the owner, top property inside the owner itself */
 	FProperty* GetMemberProperty() const
 	{
@@ -172,6 +174,9 @@ struct FPropertyAnimatorCoreData
 	/** Gets the property handler to perform operation on property without knowing the type */
 	UPropertyAnimatorCoreHandlerBase* GetPropertyHandler() const;
 
+	/** Internal use only, create the hash and display name */
+	void GeneratePropertyPath();
+
 private:
 	static void CopyPropertyValue(const FProperty* InProperty, const void* InSrc, void* OutDest);
 
@@ -193,9 +198,6 @@ private:
 	/** Uses chained properties to resolve from container to */
 	void* ContainerToValuePtr(const void* InContainer, int32 InStartPropertyIndex) const;
 
-	/** Internal use only, create the hash and display name */
-	void GeneratePropertyPath();
-
 	/** Tries to find setter function for this property */
 	bool FindSetterFunctions();
 
