MotionDesign : PropertyAnimatorCore
- Fix issue with spawnable actors with animators via template actors (sequencer)

diff --git a/Engine/Plugins/Experimental/PropertyAnimatorCore/Source/PropertyAnimatorCore/Public/Animators/PropertyAnimatorCoreBase.h b/Engine/Plugins/Experimental/PropertyAnimatorCore/Source/PropertyAnimatorCore/Public/Animators/PropertyAnimatorCoreBase.h
--- a/Engine/Plugins/Experimental/PropertyAnimatorCore/Source/PropertyAnimatorCore/Public/Animators/PropertyAnimatorCoreBase.h
+++ b/Engine/Plugins/Experimental/PropertyAnimatorCore/Source/PropertyAnimatorCore/Public/Animators/PropertyAnimatorCoreBase.h
@@ -295,11 +295,12 @@ private:
 	FName AnimatorDisplayName;
 
 	/** Context for properties linked to this Animator */
-	UPROPERTY(EditInstanceOnly, Category="Animator", EditFixedSize, meta=(EditFixedOrder))
+	// UPROPERTY(EditInstanceOnly, Category="Animator", EditFixedSize, meta=(EditFixedOrder))
+	UPROPERTY(EditInstanceOnly, NoClear, Export, Instanced, EditFixedSize, Category="Animator", meta=(EditFixedOrder))
 	TArray<TObjectPtr<UPropertyAnimatorCoreContext>> LinkedProperties;
 
 	/** Groups for properties linked to this Animator */
-	UPROPERTY(EditInstanceOnly, Instanced, Category="Animator")
+	UPROPERTY(EditInstanceOnly, NoClear, Export, Instanced, Category="Animator")
 	TArray<TObjectPtr<UPropertyAnimatorCoreGroupBase>> PropertyGroups;
 
 	/** The time source to use */
diff --git a/Engine/Plugins/Experimental/PropertyAnimatorCore/Source/PropertyAnimatorCore/Public/Components/PropertyAnimatorCoreComponent.h b/Engine/Plugins/Experimental/PropertyAnimatorCore/Source/PropertyAnimatorCore/Public/Components/PropertyAnimatorCoreComponent.h
--- a/Engine/Plugins/Experimental/PropertyAnimatorCore/Source/PropertyAnimatorCore/Public/Components/PropertyAnimatorCoreComponent.h
+++ b/Engine/Plugins/Experimental/PropertyAnimatorCore/Source/PropertyAnimatorCore/Public/Components/PropertyAnimatorCoreComponent.h
@@ -94,7 +94,7 @@ protected:
 	bool ShouldAnimatorsTick() const;
 
 	/** Animators linked to this actor, they contain only properties within this actor */
-	UPROPERTY(EditAnywhere, BlueprintReadWrite, Instanced, Setter="SetAnimators", Category="Animator", meta=(TitleProperty="AnimatorDisplayName"))
+	UPROPERTY(EditAnywhere, BlueprintReadWrite, NoClear, Export, Instanced, Setter="SetAnimators", Category="Animator", meta=(TitleProperty="AnimatorDisplayName"))
 	TArray<TObjectPtr<UPropertyAnimatorCoreBase>> PropertyAnimators;
 
 	/** Global state for all animators controlled by this component */
