MotionDesign : PropertyAnimatorCore
- Removing unused commented code

diff --git a/Engine/Plugins/Experimental/PropertyAnimatorCore/Source/PropertyAnimatorCore/Public/Animators/PropertyAnimatorCoreBase.h b/Engine/Plugins/Experimental/PropertyAnimatorCore/Source/PropertyAnimatorCore/Public/Animators/PropertyAnimatorCoreBase.h
--- a/Engine/Plugins/Experimental/PropertyAnimatorCore/Source/PropertyAnimatorCore/Public/Animators/PropertyAnimatorCoreBase.h
+++ b/Engine/Plugins/Experimental/PropertyAnimatorCore/Source/PropertyAnimatorCore/Public/Animators/PropertyAnimatorCoreBase.h
@@ -295,7 +295,6 @@ private:
 	FName AnimatorDisplayName;
 
 	/** Context for properties linked to this Animator */
-	// UPROPERTY(EditInstanceOnly, Category="Animator", EditFixedSize, meta=(EditFixedOrder))
 	UPROPERTY(EditInstanceOnly, NoClear, Export, Instanced, EditFixedSize, Category="Animator", meta=(EditFixedOrder))
 	TArray<TObjectPtr<UPropertyAnimatorCoreContext>> LinkedProperties;
 
