MotionDesign : PropertyAnimator
- Adding submenu with supported properties from preset to let user pick specific ones
- Adding an empty animator from the operator stack "Add Animator" action just like modifiers

diff --git a/Engine/Plugins/Experimental/Avalanche/Source/AvalanchePropertyAnimatorEditor/Private/Outliner/AvaPropertyAnimatorEditorOutlinerContextMenu.cpp b/Engine/Plugins/Experimental/Avalanche/Source/AvalanchePropertyAnimatorEditor/Private/Outliner/AvaPropertyAnimatorEditorOutlinerContextMenu.cpp
--- a/Engine/Plugins/Experimental/Avalanche/Source/AvalanchePropertyAnimatorEditor/Private/Outliner/AvaPropertyAnimatorEditorOutlinerContextMenu.cpp
+++ b/Engine/Plugins/Experimental/Avalanche/Source/AvalanchePropertyAnimatorEditor/Private/Outliner/AvaPropertyAnimatorEditorOutlinerContextMenu.cpp
@@ -39,7 +39,7 @@ void FAvaPropertyAnimatorEditorOutlinerContextMenu::OnExtendOutlinerContextMenu(
 	FPropertyAnimatorCoreEditorMenuOptions MenuOptions(
 		{
 			EPropertyAnimatorCoreEditorMenuType::Edit
-			, EPropertyAnimatorCoreEditorMenuType::New
+			, EPropertyAnimatorCoreEditorMenuType::NewAdvanced
 			, EPropertyAnimatorCoreEditorMenuType::Enable
 			, EPropertyAnimatorCoreEditorMenuType::Disable
 			, EPropertyAnimatorCoreEditorMenuType::Delete
diff --git a/Engine/Plugins/Experimental/PropertyAnimator/Source/PropertyAnimator/Private/Presets/PropertyAnimatorPresetLocationYZ.cpp b/Engine/Plugins/Experimental/PropertyAnimator/Source/PropertyAnimator/Private/Presets/PropertyAnimatorPresetLocationYZ.cpp
deleted file mode 100644
diff --git a/Engine/Plugins/Experimental/PropertyAnimator/Source/PropertyAnimator/Private/Presets/PropertyAnimatorPresetLocationYZ.h b/Engine/Plugins/Experimental/PropertyAnimator/Source/PropertyAnimator/Private/Presets/PropertyAnimatorPresetLocationYZ.h
deleted file mode 100644
diff --git a/Engine/Plugins/Experimental/PropertyAnimatorCore/Source/PropertyAnimatorCore/Private/Presets/PropertyAnimatorCorePresetBase.cpp b/Engine/Plugins/Experimental/PropertyAnimatorCore/Source/PropertyAnimatorCore/Private/Presets/PropertyAnimatorCorePresetBase.cpp
--- a/Engine/Plugins/Experimental/PropertyAnimatorCore/Source/PropertyAnimatorCore/Private/Presets/PropertyAnimatorCorePresetBase.cpp
+++ b/Engine/Plugins/Experimental/PropertyAnimatorCore/Source/PropertyAnimatorCore/Private/Presets/PropertyAnimatorCorePresetBase.cpp
@@ -13,6 +13,7 @@ void UPropertyAnimatorCorePresetBase::GetSupportedPresetProperties(const AActor*
 {
 	TSet<FPropertyAnimatorCoreData> PresetProperties;
 	GetPresetProperties(InActor, InAnimator, PresetProperties);
+	OutProperties.Empty(PresetProperties.Num());
 
 	if (PresetProperties.IsEmpty())
 	{
@@ -105,3 +106,25 @@ bool UPropertyAnimatorCorePresetBase::UnapplyPreset(UPropertyAnimatorCoreBase* I
 
 	return true;
 }
+
+void UPropertyAnimatorCorePresetBase::GetAppliedPresetProperties(const UPropertyAnimatorCoreBase* InAnimator, TSet<FPropertyAnimatorCoreData>& OutSupportedProperties, TSet<FPropertyAnimatorCoreData>& OutAppliedProperties)
+{
+	OutSupportedProperties.Empty();
+	OutAppliedProperties.Empty();
+
+	if (!IsValid(InAnimator) || InAnimator->IsTemplate())
+	{
+		return;
+	}
+
+	GetSupportedPresetProperties(InAnimator->GetAnimatorActor(), InAnimator, OutSupportedProperties);
+	OutAppliedProperties.Reserve(OutSupportedProperties.Num());
+
+	for (const FPropertyAnimatorCoreData& Property : OutSupportedProperties)
+	{
+		if (InAnimator->IsPropertyLinked(Property))
+		{
+			OutAppliedProperties.Add(Property);
+		}
+	}
+}
diff --git a/Engine/Plugins/Experimental/PropertyAnimatorCore/Source/PropertyAnimatorCore/Public/Presets/PropertyAnimatorCorePresetBase.h b/Engine/Plugins/Experimental/PropertyAnimatorCore/Source/PropertyAnimatorCore/Public/Presets/PropertyAnimatorCorePresetBase.h
--- a/Engine/Plugins/Experimental/PropertyAnimatorCore/Source/PropertyAnimatorCore/Public/Presets/PropertyAnimatorCorePresetBase.h
+++ b/Engine/Plugins/Experimental/PropertyAnimatorCore/Source/PropertyAnimatorCore/Public/Presets/PropertyAnimatorCorePresetBase.h
@@ -52,6 +52,9 @@ public:
 	/** Un applies this preset on an animator */
 	PROPERTYANIMATORCORE_API virtual bool UnapplyPreset(UPropertyAnimatorCoreBase* InAnimator);
 
+	/** Gets the supported and applied properties for an animator */
+	PROPERTYANIMATORCORE_API virtual void GetAppliedPresetProperties(const UPropertyAnimatorCoreBase* InAnimator, TSet<FPropertyAnimatorCoreData>& OutSupportedProperties, TSet<FPropertyAnimatorCoreData>& OutAppliedProperties);
+
 	/** Called when this preset is applied on the animator */
 	virtual void OnPresetApplied(UPropertyAnimatorCoreBase* InAnimator, const TSet<FPropertyAnimatorCoreData>& InProperties) {}
 
diff --git a/Engine/Plugins/Experimental/PropertyAnimatorCore/Source/PropertyAnimatorCoreEditor/Private/Customizations/PropertyAnimatorCoreEditorStackCustomization.cpp b/Engine/Plugins/Experimental/PropertyAnimatorCore/Source/PropertyAnimatorCoreEditor/Private/Customizations/PropertyAnimatorCoreEditorStackCustomization.cpp
--- a/Engine/Plugins/Experimental/PropertyAnimatorCore/Source/PropertyAnimatorCoreEditor/Private/Customizations/PropertyAnimatorCoreEditorStackCustomization.cpp
+++ b/Engine/Plugins/Experimental/PropertyAnimatorCore/Source/PropertyAnimatorCoreEditor/Private/Customizations/PropertyAnimatorCoreEditorStackCustomization.cpp
@@ -258,6 +258,7 @@ void UPropertyAnimatorCoreEditorStackCustomization::CustomizeItemHeader(const FO
 		else
 		{
 			HeaderLabel = Animator->GetAnimatorOriginalName().ToString();
+			HeaderLabel += TEXT(" (") + FString::FromInt(InItem->GetValueCount()) + TEXT(")");
 		}
 
 		static const FLinearColor AnimatorColor = FLinearColor(FColor::Orange).Desaturate(0.25);
@@ -396,7 +397,7 @@ void UPropertyAnimatorCoreEditorStackCustomization::FillAddAnimatorMenuSection(U
 	}
 
 	const FPropertyAnimatorCoreEditorMenuContext MenuContext(ContextObjects, {});
-	const FPropertyAnimatorCoreEditorMenuOptions MenuOptions({EPropertyAnimatorCoreEditorMenuType::New});
+	const FPropertyAnimatorCoreEditorMenuOptions MenuOptions({EPropertyAnimatorCoreEditorMenuType::NewSimple});
 	AnimatorEditorSubsystem->FillAnimatorMenu(InToolMenu, MenuContext, MenuOptions);
 }
 
diff --git a/Engine/Plugins/Experimental/PropertyAnimatorCore/Source/PropertyAnimatorCoreEditor/Private/Menus/PropertyAnimatorCoreEditorMenu.cpp b/Engine/Plugins/Experimental/PropertyAnimatorCore/Source/PropertyAnimatorCoreEditor/Private/Menus/PropertyAnimatorCoreEditorMenu.cpp
--- a/Engine/Plugins/Experimental/PropertyAnimatorCore/Source/PropertyAnimatorCoreEditor/Private/Menus/PropertyAnimatorCoreEditorMenu.cpp
+++ b/Engine/Plugins/Experimental/PropertyAnimatorCore/Source/PropertyAnimatorCoreEditor/Private/Menus/PropertyAnimatorCoreEditorMenu.cpp
@@ -5,6 +5,7 @@
 #include "Animators/PropertyAnimatorCoreBase.h"
 #include "Components/PropertyAnimatorCoreComponent.h"
 #include "Engine/World.h"
+#include "Misc/Optional.h"
 #include "PropertyHandle.h"
 #include "Presets/PropertyAnimatorCorePresetBase.h"
 #include "Styling/SlateIconFinder.h"
@@ -74,6 +75,8 @@
 
 	constexpr bool bCloseMenuAfterSelection = false;
 	constexpr bool bOpenOnClick = false;
+	const bool bAdvancedMenu = InMenuData->GetOptions().IsMenuType(EPropertyAnimatorCoreEditorMenuType::NewAdvanced);
+	const TSet<AActor*>& ContextActors = InMenuData->GetContext().GetActors();
 
 	FToolMenuSection& NewAnimatorsSection = InMenu->FindOrAddSection(TEXT("NewAnimators"), LOCTEXT("NewAnimators.Label", "New Animators"));
 
@@ -81,17 +84,33 @@
 	{
 		const FName MenuName = NewAnimator->GetAnimatorOriginalName();
 		const FText MenuLabel = FText::FromName(MenuName);
+		const FText MenuTooltip = LOCTEXT("NewAnimator.Tooltip", "Create a new animator");
 		const FSlateIcon MenuIcon = FSlateIconFinder::FindIconForClass(NewAnimator->GetClass());
 
-		NewAnimatorsSection.AddSubMenu(
-			MenuName
-			, MenuLabel
-			, LOCTEXT("NewAnimator.Tooltip", "Create a new animator")
-			, FNewToolMenuDelegate::CreateLambda(&FillNewAnimatorSubmenu, NewAnimator, InMenuData)
-			, bOpenOnClick
-			, MenuIcon
-			, bCloseMenuAfterSelection
-		);
+		if (bAdvancedMenu)
+		{
+			NewAnimatorsSection.AddSubMenu(
+				MenuName
+				, MenuLabel
+				, MenuTooltip
+				, FNewToolMenuDelegate::CreateLambda(&FillNewAnimatorSubmenu, NewAnimator, InMenuData)
+				, bOpenOnClick
+				, MenuIcon
+				, bCloseMenuAfterSelection
+			);
+		}
+		else
+		{
+			UPropertyAnimatorCorePresetBase* EmptyPreset = nullptr;
+
+			NewAnimatorsSection.AddMenuEntry(
+				MenuName
+				, MenuLabel
+				, MenuTooltip
+				, MenuIcon
+				, FExecuteAction::CreateLambda(&ExecuteNewAnimatorPresetAction, NewAnimator, ContextActors, EmptyPreset, InMenuData)
+			);
+		}
 	}
 }
 
@@ -582,6 +601,56 @@
 	Subsystem->RemoveAnimator(InAnimator, InMenuData->GetOptions().ShouldTransact());
 }
 
+ECheckBoxState UE::PropertyAnimatorCoreEditor::Menu::GetAnimatorPresetState(const UPropertyAnimatorCoreBase* InAnimator, UPropertyAnimatorCorePresetBase* InPreset)
+{
+	if (!IsValid(InAnimator) || InAnimator->IsTemplate() || !InPreset)
+	{
+		return ECheckBoxState::Unchecked;
+	}
+
+	TSet<FPropertyAnimatorCoreData> SupportedProperties;
+	TSet<FPropertyAnimatorCoreData> AppliedProperties;
+	InPreset->GetAppliedPresetProperties(InAnimator, SupportedProperties, AppliedProperties);
+
+	if (!SupportedProperties.IsEmpty() && SupportedProperties.Num() == AppliedProperties.Num())
+	{
+		return ECheckBoxState::Checked;
+	}
+
+	return !AppliedProperties.IsEmpty() ? ECheckBoxState::Undetermined : ECheckBoxState::Unchecked;
+}
+
+ECheckBoxState UE::PropertyAnimatorCoreEditor::Menu::GetLastAnimatorCreatedPresetState(const UPropertyAnimatorCoreBase* InAnimator, UPropertyAnimatorCorePresetBase* InPreset, TSharedRef<FPropertyAnimatorCoreEditorMenuData> InMenuData)
+{
+	if (!IsValid(InAnimator)
+		|| !InAnimator->IsTemplate()
+		|| !InPreset
+		|| !InMenuData->ContainsAnyLastCreatedAnimator())
+	{
+		return ECheckBoxState::Unchecked;
+	}
+
+	TOptional<ECheckBoxState> AnimatorsState;
+
+	for (const UPropertyAnimatorCoreBase* LastCreatedAnimator : InMenuData->GetLastCreatedAnimators())
+	{
+		if (LastCreatedAnimator->GetClass() == InAnimator->GetClass())
+		{
+			if (!AnimatorsState.IsSet())
+			{
+				AnimatorsState = GetAnimatorPresetState(LastCreatedAnimator, InPreset);
+			}
+			else if (GetAnimatorPresetState(LastCreatedAnimator, InPreset) != AnimatorsState.GetValue())
+			{
+				AnimatorsState = ECheckBoxState::Undetermined;
+				break;
+			}
+		}
+	}
+
+	return AnimatorsState.Get(ECheckBoxState::Unchecked);
+}
+
 bool UE::PropertyAnimatorCoreEditor::Menu::IsAnimatorPresetLinked(UPropertyAnimatorCoreBase* InAnimator, UPropertyAnimatorCorePresetBase* InPreset)
 {
 	if (!IsValid(InAnimator) || InAnimator->IsTemplate() || !InPreset)
@@ -632,6 +701,9 @@
 			SupportedPresets = SupportedPresets.Intersect(Subsystem->GetSupportedPresets(Actor, InAnimator));
 		}
 
+		constexpr bool bCloseMenuAfterSelection = false;
+		constexpr bool bOpenOnClick = false;
+
 		for (UPropertyAnimatorCorePresetBase* SupportedPreset : SupportedPresets)
 		{
 			if (!SupportedPreset)
@@ -642,33 +714,14 @@
 			const FString MenuName = SupportedPreset->GetPresetName().ToString();
 			const FText MenuLabel = FText::FromString(SupportedPreset->GetPresetDisplayName());
 
-			// Create action (creates an animator and links the property)
-			PresetSection.AddMenuEntry(
+			PresetSection.AddSubMenu(
 				FName(TEXT("Create") + MenuName)
 				, MenuLabel
 				, LOCTEXT("NewAnimatorPresetSection.Tooltip", "Create this animator using this preset")
+				, FNewToolMenuDelegate::CreateLambda(&FillNewPresetAnimatorSubmenu, InAnimator, SupportedPreset, InMenuData)
+				, bOpenOnClick
 				, FSlateIcon()
-				, FUIAction(
-					FExecuteAction::CreateLambda(&ExecuteNewAnimatorPresetAction, InAnimator, ContextActors, SupportedPreset, InMenuData)
-					, FCanExecuteAction()
-					, FIsActionChecked()
-					, FIsActionButtonVisible::CreateLambda(&IsLastAnimatorCreatedActionHidden, InAnimator, InMenuData)
-				)
+				, bCloseMenuAfterSelection
 			);
-
-			// Link action (links the property to last created animator)
-			PresetSection.AddMenuEntry(
-				FName(TEXT("Apply") + MenuName)
-				, MenuLabel
-				, LOCTEXT("ApplyLastCreatedAnimatorPresetSection.Tooltip", "Apply this preset to the last created animator")
-				, FSlateIcon()
-				, FUIAction(
-					FExecuteAction::CreateLambda(&ExecuteApplyLastCreatedAnimatorPresetAction, InAnimator, SupportedPreset, InMenuData)
-					, FCanExecuteAction()
-					, FIsActionChecked::CreateLambda(&IsLastAnimatorCreatedPresetLinked, InAnimator, SupportedPreset, InMenuData)
-					, FIsActionButtonVisible::CreateLambda(&IsLastAnimatorCreatedActionVisible, InAnimator, InMenuData)
-				)
-				, EUserInterfaceActionType::ToggleButton
-			);
 		}
 	}
@@ -745,6 +798,9 @@
 		SupportedPresets = SupportedPresets.Intersect(Subsystem->GetSupportedPresets(Actor, InAnimator));
 	}
 
+	constexpr bool bCloseMenuAfterSelection = false;
+    constexpr bool bOpenOnClick = false;
+
 	for (UPropertyAnimatorCorePresetBase* SupportedPreset : SupportedPresets)
 	{
 		if (!SupportedPreset)
@@ -754,18 +810,16 @@
 
 		const FName MenuName = SupportedPreset->GetPresetName();
 		const FText MenuLabel = FText::FromString(SupportedPreset->GetPresetDisplayName());
+		const FText MenuTooltip = LOCTEXT("LinkAnimatorPresetSection.Tooltip", "Link or unlink a preset from this animator");
 
-		PresetSection.AddMenuEntry(
+		PresetSection.AddSubMenu(
 			MenuName
 			, MenuLabel
-			, LOCTEXT("LinkAnimatorPresetSection.Tooltip", "Link or unlink a preset from this animator")
+			, MenuTooltip
+			, FNewToolMenuDelegate::CreateLambda(&FillPresetAnimatorSubmenu, InAnimator, SupportedPreset, InMenuData)
+			, bOpenOnClick
 			, FSlateIcon()
-			, FUIAction(
-				FExecuteAction::CreateLambda(&ExecuteLinkAnimatorPresetAction, InAnimator, SupportedPreset, InMenuData)
-				, FCanExecuteAction()
-				, FIsActionChecked::CreateStatic(&IsAnimatorPresetLinked, InAnimator, SupportedPreset)
-			)
-			, EUserInterfaceActionType::ToggleButton
+			, bCloseMenuAfterSelection
 		);
 	}
 
@@ -797,6 +851,174 @@
 	}
 }
 
+void UE::PropertyAnimatorCoreEditor::Menu::FillPresetAnimatorSubmenu(UToolMenu* InMenu, UPropertyAnimatorCoreBase* InAnimator, UPropertyAnimatorCorePresetBase* InPreset, TSharedRef<FPropertyAnimatorCoreEditorMenuData> InMenuData)
+{
+	if (!InMenu)
+	{
+		return;
+	}
+
+	const FToolMenuEntry AllPropertiesEntry = FToolMenuEntry::InitMenuEntry(
+		TEXT("All")
+		, LOCTEXT("LinkAllPresetProperty.Label", "All")
+		, LOCTEXT("LinkAllPresetProperty.Tooltip", "Link all properties from this preset")
+		, FSlateIcon()
+		, FUIAction(
+			FExecuteAction::CreateLambda(&ExecuteLinkAnimatorPresetAction, InAnimator, InPreset, InMenuData)
+			, FCanExecuteAction()
+			, FGetActionCheckState::CreateLambda(&GetAnimatorPresetState, InAnimator, InPreset)
+		)
+		, EUserInterfaceActionType::ToggleButton
+	);
+
+	// Add all preset properties option
+	InMenu->AddMenuEntry(
+		AllPropertiesEntry.Name,
+		AllPropertiesEntry
+	);
+
+	const FToolMenuEntry SeparatorEntry = FToolMenuEntry::InitSeparator(TEXT("PresetSeparator"));
+	InMenu->AddMenuEntry(SeparatorEntry.Name, SeparatorEntry);
+
+	TSet<FPropertyAnimatorCoreData> SupportedProperties;
+	InPreset->GetSupportedPresetProperties(InAnimator->GetAnimatorActor(), InAnimator, SupportedProperties);
+
+	for (const FPropertyAnimatorCoreData& SupportedProperty : SupportedProperties)
+	{
+		const FName MenuName = SupportedProperty.GetPropertyDisplayName();
+		const FText MenuLabel = FText::FromName(MenuName);
+		const FText MenuTooltip = LOCTEXT("LinkPresetProperty.Tooltip", "Link this preset property");
+
+		const FToolMenuEntry SupportedPropertyEntry = FToolMenuEntry::InitMenuEntry(
+			MenuName
+			, MenuLabel
+			, MenuTooltip
+			, FSlateIcon()
+			, FUIAction(
+				FExecuteAction::CreateLambda(&ExecuteLinkAnimatorPropertyAction, InAnimator, SupportedProperty, InMenuData)
+				, FCanExecuteAction()
+				, FIsActionChecked::CreateLambda(&IsAnimatorPropertyLinked, InAnimator, SupportedProperty)
+			)
+			, EUserInterfaceActionType::ToggleButton
+		);
+
+		InMenu->AddMenuEntry(SupportedPropertyEntry.Name, SupportedPropertyEntry);
+	}
+}
+
+void UE::PropertyAnimatorCoreEditor::Menu::FillNewPresetAnimatorSubmenu(UToolMenu* InMenu, UPropertyAnimatorCoreBase* InAnimator, UPropertyAnimatorCorePresetBase* InPreset, TSharedRef<FPropertyAnimatorCoreEditorMenuData> InMenuData)
+{
+	if (!InMenu)
+	{
+		return;
+	}
+
+	const TSet<AActor*>& ContextActors = InMenuData->GetContext().GetActors();
+
+	int32 Index = 0;
+	TSet<FPropertyAnimatorCoreData> PresetProperties;
+	for (const AActor* ContextActor : ContextActors)
+	{
+		TSet<FPropertyAnimatorCoreData> SupportedProperties;
+		InPreset->GetSupportedPresetProperties(ContextActor, InAnimator, SupportedProperties);
+
+		if (Index++ == 0)
+		{
+			PresetProperties = SupportedProperties;
+		}
+		else
+		{
+			PresetProperties = PresetProperties.Intersect(SupportedProperties);
+		}
+	}
+
+	const FToolMenuEntry CreateAllPropertiesEntry = FToolMenuEntry::InitMenuEntry(
+		TEXT("CreateAllProperties")
+		, LOCTEXT("NewAnimatorPresetSection.Label", "All")
+		, LOCTEXT("NewAnimatorPresetSection.Tooltip", "Create this animator using this preset")
+		, FSlateIcon()
+		, FUIAction(
+			FExecuteAction::CreateLambda(&ExecuteNewAnimatorPresetAction, InAnimator, ContextActors, InPreset, InMenuData)
+			, FCanExecuteAction()
+			, FIsActionChecked()
+			, FIsActionButtonVisible::CreateLambda(&IsLastAnimatorCreatedActionHidden, InAnimator, InMenuData)
+		)
+	);
+
+	// Create action (creates an animator and links the property)
+	InMenu->AddMenuEntry(
+		CreateAllPropertiesEntry.Name
+		, CreateAllPropertiesEntry
+	);
+
+	const FToolMenuEntry LinkAllPropertiesEntry = FToolMenuEntry::InitMenuEntry(
+		TEXT("LinkAllProperties")
+		, LOCTEXT("ApplyLastCreatedAnimatorPresetSection.Label", "All")
+		, LOCTEXT("ApplyLastCreatedAnimatorPresetSection.Tooltip", "Apply this preset to the last created animator")
+		, FSlateIcon()
+		, FUIAction(
+			FExecuteAction::CreateLambda(&ExecuteApplyLastCreatedAnimatorPresetAction, InAnimator, InPreset, InMenuData)
+			, FCanExecuteAction()
+			, FGetActionCheckState::CreateLambda(&GetLastAnimatorCreatedPresetState, InAnimator, InPreset, InMenuData)
+			, FIsActionButtonVisible::CreateLambda(&IsLastAnimatorCreatedActionVisible, InAnimator, InMenuData)
+		)
+		, EUserInterfaceActionType::ToggleButton
+	);
+
+	// Link action (links the property to last created animator)
+	InMenu->AddMenuEntry(
+		LinkAllPropertiesEntry.Name
+		, LinkAllPropertiesEntry
+	);
+
+	const FToolMenuEntry SeparatorEntry = FToolMenuEntry::InitSeparator(TEXT("PresetSeparator"));
+	InMenu->AddMenuEntry(SeparatorEntry.Name, SeparatorEntry);
+
+	for (const FPropertyAnimatorCoreData& PresetProperty : PresetProperties)
+	{
+		const FText MenuLabel = FText::FromName(PresetProperty.GetPropertyDisplayName());
+
+		const FToolMenuEntry CreatePropertyEntry = FToolMenuEntry::InitMenuEntry(
+			FName(TEXT("CreateProperty") + PresetProperty.GetPropertyDisplayName().ToString())
+			, MenuLabel
+			, LOCTEXT("CreateAnimatorPresetProperty.Tooltip", "Create this animator using this preset property")
+			, FSlateIcon()
+			, FUIAction(
+				FExecuteAction::CreateLambda(&ExecuteNewAnimatorPropertyAction, InAnimator, PresetProperty, InMenuData)
+				, FCanExecuteAction()
+				, FIsActionChecked()
+				, FIsActionButtonVisible::CreateLambda(&IsLastAnimatorCreatedActionHidden, InAnimator, InMenuData)
+			)
+		);
+
+		// Create action (creates an animator and links the property)
+		InMenu->AddMenuEntry(
+			CreatePropertyEntry.Name
+			, CreatePropertyEntry
+		);
+
+		const FToolMenuEntry LinkPropertyEntry = FToolMenuEntry::InitMenuEntry(
+			FName(TEXT("LinkProperty") + PresetProperty.GetPropertyDisplayName().ToString())
+			, MenuLabel
+			, LOCTEXT("LinkAnimatorPresetProperty.Tooltip", "Link this preset property to the last created animator")
+			, FSlateIcon()
+			, FUIAction(
+				FExecuteAction::CreateLambda(&ExecuteLinkLastCreatedAnimatorPropertyAction, InAnimator, PresetProperty, InMenuData)
+				, FCanExecuteAction()
+				, FIsActionChecked::CreateLambda(&IsLastAnimatorCreatedPropertyLinked, InAnimator, PresetProperty, InMenuData)
+				, FIsActionButtonVisible::CreateLambda(&IsLastAnimatorCreatedActionVisible, InAnimator, InMenuData)
+			)
+			, EUserInterfaceActionType::ToggleButton
+		);
+
+		// Link action (links the property to last created animator)
+		InMenu->AddMenuEntry(
+			LinkPropertyEntry.Name
+			, LinkPropertyEntry
+		);
+	}
+}
+
 bool UE::PropertyAnimatorCoreEditor::Menu::IsAnimatorPropertyLinked(const UPropertyAnimatorCoreBase* InAnimator, const FPropertyAnimatorCoreData& InProperty)
 {
 	if (!IsValid(InAnimator) || InAnimator->IsTemplate())
@@ -814,25 +1036,3 @@
 		|| !InMenuData->ContainsAnyLastCreatedAnimator())
 	{
 		return false;
-	}
-
-	for (const UPropertyAnimatorCoreBase* LastCreatedAnimator : InMenuData->GetLastCreatedAnimators())
-	{
-		if (LastCreatedAnimator->GetClass() != InAnimator->GetClass()
-			|| !LastCreatedAnimator->IsPropertyLinked(InProperty))
-		{
-			return false;
-		}
-	}
-
-	return true;
-}
-
-bool UE::PropertyAnimatorCoreEditor::Menu::IsLastAnimatorCreatedPresetLinked(const UPropertyAnimatorCoreBase* InAnimator, UPropertyAnimatorCorePresetBase* InPreset, TSharedRef<FPropertyAnimatorCoreEditorMenuData> InMenuData)
-{
-	if (!IsValid(InAnimator)
-		|| !InAnimator->IsTemplate()
-		|| !InPreset
-		|| !InMenuData->ContainsAnyLastCreatedAnimator())
-	{
-		return false;
@@ -841,7 +1041,7 @@
 	for (const UPropertyAnimatorCoreBase* LastCreatedAnimator : InMenuData->GetLastCreatedAnimators())
 	{
 		if (LastCreatedAnimator->GetClass() != InAnimator->GetClass()
-			|| !InPreset->IsPresetApplied(LastCreatedAnimator))
+			|| !LastCreatedAnimator->IsPropertyLinked(InProperty))
 		{
 			return false;
 		}
diff --git a/Engine/Plugins/Experimental/PropertyAnimatorCore/Source/PropertyAnimatorCoreEditor/Private/Menus/PropertyAnimatorCoreEditorMenu.h b/Engine/Plugins/Experimental/PropertyAnimatorCore/Source/PropertyAnimatorCoreEditor/Private/Menus/PropertyAnimatorCoreEditorMenu.h
--- a/Engine/Plugins/Experimental/PropertyAnimatorCore/Source/PropertyAnimatorCoreEditor/Private/Menus/PropertyAnimatorCoreEditorMenu.h
+++ b/Engine/Plugins/Experimental/PropertyAnimatorCore/Source/PropertyAnimatorCoreEditor/Private/Menus/PropertyAnimatorCoreEditorMenu.h
@@ -39,6 +39,10 @@ namespace UE::PropertyAnimatorCoreEditor::Menu
 
 	void FillLinkAnimatorSubmenu(UToolMenu* InMenu, UPropertyAnimatorCoreBase* InAnimator, TSharedRef<FPropertyAnimatorCoreEditorMenuData> InMenuData);
 
+	void FillPresetAnimatorSubmenu(UToolMenu* InMenu, UPropertyAnimatorCoreBase* InAnimator, UPropertyAnimatorCorePresetBase* InPreset, TSharedRef<FPropertyAnimatorCoreEditorMenuData> InMenuData);
+
+	void FillNewPresetAnimatorSubmenu(UToolMenu* InMenu, UPropertyAnimatorCoreBase* InAnimator, UPropertyAnimatorCorePresetBase* InPreset, TSharedRef<FPropertyAnimatorCoreEditorMenuData> InMenuData);
+
 	/** Execute */
 
 	void ExecuteEditAnimatorAction(AActor* InActor);
@@ -67,6 +71,10 @@ namespace UE::PropertyAnimatorCoreEditor::Menu
 
 	/** Check */
 
+	ECheckBoxState GetAnimatorPresetState(const UPropertyAnimatorCoreBase* InAnimator, UPropertyAnimatorCorePresetBase* InPreset);
+
+	ECheckBoxState GetLastAnimatorCreatedPresetState(const UPropertyAnimatorCoreBase* InAnimator, UPropertyAnimatorCorePresetBase* InPreset, TSharedRef<FPropertyAnimatorCoreEditorMenuData> InMenuData);
+
 	bool IsAnimatorPresetLinked(UPropertyAnimatorCoreBase* InAnimator, UPropertyAnimatorCorePresetBase* InPreset);
 
 	bool IsAnimatorPropertyLinked(const UPropertyAnimatorCoreBase* InAnimator, const FPropertyAnimatorCoreData& InProperty);
diff --git a/Engine/Plugins/Experimental/PropertyAnimatorCore/Source/PropertyAnimatorCoreEditor/Private/Subsystems/PropertyAnimatorCoreEditorSubsystem.cpp b/Engine/Plugins/Experimental/PropertyAnimatorCore/Source/PropertyAnimatorCoreEditor/Private/Subsystems/PropertyAnimatorCoreEditorSubsystem.cpp
--- a/Engine/Plugins/Experimental/PropertyAnimatorCore/Source/PropertyAnimatorCoreEditor/Private/Subsystems/PropertyAnimatorCoreEditorSubsystem.cpp
+++ b/Engine/Plugins/Experimental/PropertyAnimatorCore/Source/PropertyAnimatorCoreEditor/Private/Subsystems/PropertyAnimatorCoreEditorSubsystem.cpp
@@ -85,14 +85,30 @@ bool UPropertyAnimatorCoreEditorSubsystem::FillAnimatorMenu(UToolMenu* InMenu, c
 		}
 	}
 
-	if (InOptions.IsMenuType(EPropertyAnimatorCoreEditorMenuType::New) && InContext.ContainsAnyActor())
+	if (InOptions.IsMenuType(EPropertyAnimatorCoreEditorMenuType::NewSimple) && InContext.ContainsAnyActor())
 	{
 		if (InOptions.ShouldCreateSubMenu())
 		{
 			AnimatorSection->AddSubMenu(
-				TEXT("NewAnimatorMenu"),
-				LOCTEXT("NewAnimatorMenu.Label", "Add Animators"),
-				LOCTEXT("NewAnimatorMenu.Tooltip", "Add animators to the selection"),
+				TEXT("NewSimpleAnimatorMenu"),
+				LOCTEXT("NewSimpleAnimatorMenu.Label", "Add Animators"),
+				LOCTEXT("NewSimpleAnimatorMenu.Tooltip", "Add animators to the selection"),
+				FNewToolMenuDelegate::CreateLambda(&UE::PropertyAnimatorCoreEditor::Menu::FillNewAnimatorSection, LastMenuData.ToSharedRef()));
+		}
+		else
+		{
+			UE::PropertyAnimatorCoreEditor::Menu::FillNewAnimatorSection(InMenu, LastMenuData.ToSharedRef());
+		}
+	}
+
+	if (InOptions.IsMenuType(EPropertyAnimatorCoreEditorMenuType::NewAdvanced) && InContext.ContainsAnyActor())
+	{
+		if (InOptions.ShouldCreateSubMenu())
+		{
+			AnimatorSection->AddSubMenu(
+				TEXT("NewAdvancedAnimatorMenu"),
+				LOCTEXT("NewAdvancedAnimatorMenu.Label", "Add Animators"),
+				LOCTEXT("NewAdvancedAnimatorMenu.Tooltip", "Add animators to the selection"),
 				FNewToolMenuDelegate::CreateLambda(&UE::PropertyAnimatorCoreEditor::Menu::FillNewAnimatorSection, LastMenuData.ToSharedRef()));
 		}
 		else
@@ -394,7 +410,7 @@ void UPropertyAnimatorCoreEditorSubsystem::FillAnimatorExtensionSection(UToolMen
 	}
 
 	const FPropertyAnimatorCoreEditorMenuContext MenuContext({}, {Context->GetPropertyData()});
-	const FPropertyAnimatorCoreEditorMenuOptions MenuOptions({EPropertyAnimatorCoreEditorMenuType::Edit, EPropertyAnimatorCoreEditorMenuType::New, EPropertyAnimatorCoreEditorMenuType::Existing});
+	const FPropertyAnimatorCoreEditorMenuOptions MenuOptions({EPropertyAnimatorCoreEditorMenuType::Edit, EPropertyAnimatorCoreEditorMenuType::NewAdvanced, EPropertyAnimatorCoreEditorMenuType::Existing});
 	FillAnimatorMenu(InToolMenu, MenuContext, MenuOptions);
 }
 
@@ -421,7 +437,7 @@ void UPropertyAnimatorCoreEditorSubsystem::FillAnimatorRowContextSection(UToolMe
 	}
 
 	const FPropertyAnimatorCoreEditorMenuContext MenuContext({}, {PropertyData.GetValue()});
-	const FPropertyAnimatorCoreEditorMenuOptions MenuOptions({EPropertyAnimatorCoreEditorMenuType::Edit, EPropertyAnimatorCoreEditorMenuType::New, EPropertyAnimatorCoreEditorMenuType::Existing});
+	const FPropertyAnimatorCoreEditorMenuOptions MenuOptions({EPropertyAnimatorCoreEditorMenuType::Edit, EPropertyAnimatorCoreEditorMenuType::NewAdvanced, EPropertyAnimatorCoreEditorMenuType::Existing});
 	FillAnimatorMenu(InToolMenu, MenuContext, MenuOptions);
 }
 
diff --git a/Engine/Plugins/Experimental/PropertyAnimatorCore/Source/PropertyAnimatorCoreEditor/Private/Widgets/SPropertyAnimatorCoreEditorEditPanel.cpp b/Engine/Plugins/Experimental/PropertyAnimatorCore/Source/PropertyAnimatorCoreEditor/Private/Widgets/SPropertyAnimatorCoreEditorEditPanel.cpp
--- a/Engine/Plugins/Experimental/PropertyAnimatorCore/Source/PropertyAnimatorCoreEditor/Private/Widgets/SPropertyAnimatorCoreEditorEditPanel.cpp
+++ b/Engine/Plugins/Experimental/PropertyAnimatorCore/Source/PropertyAnimatorCoreEditor/Private/Widgets/SPropertyAnimatorCoreEditorEditPanel.cpp
@@ -298,7 +298,7 @@ void SPropertyAnimatorCoreEditorEditPanel::FillAnimatorContextSection(UToolMenu*
 	}
 
 	const FPropertyAnimatorCoreEditorMenuContext MenuContext({}, {Context->GetPropertyData()});
-	const FPropertyAnimatorCoreEditorMenuOptions MenuOptions({EPropertyAnimatorCoreEditorMenuType::New});
+	const FPropertyAnimatorCoreEditorMenuOptions MenuOptions({EPropertyAnimatorCoreEditorMenuType::NewAdvanced});
 	EditorSubsystem->FillAnimatorMenu(InToolMenu, MenuContext, MenuOptions);
 }
 
diff --git a/Engine/Plugins/Experimental/PropertyAnimatorCore/Source/PropertyAnimatorCoreEditor/Private/Widgets/TableRows/SPropertyAnimatorCoreEditorPropertiesViewTableRow.cpp b/Engine/Plugins/Experimental/PropertyAnimatorCore/Source/PropertyAnimatorCoreEditor/Private/Widgets/TableRows/SPropertyAnimatorCoreEditorPropertiesViewTableRow.cpp
--- a/Engine/Plugins/Experimental/PropertyAnimatorCore/Source/PropertyAnimatorCoreEditor/Private/Widgets/TableRows/SPropertyAnimatorCoreEditorPropertiesViewTableRow.cpp
+++ b/Engine/Plugins/Experimental/PropertyAnimatorCore/Source/PropertyAnimatorCoreEditor/Private/Widgets/TableRows/SPropertyAnimatorCoreEditorPropertiesViewTableRow.cpp
@@ -306,7 +306,7 @@ void SPropertyAnimatorCoreEditorPropertiesViewTableRow::FillNewAnimatorPropertyV
 	}
 
 	const FPropertyAnimatorCoreEditorMenuContext MenuContext({}, {Context->GetPropertyData()});
-	const FPropertyAnimatorCoreEditorMenuOptions MenuOptions({EPropertyAnimatorCoreEditorMenuType::New});
+	const FPropertyAnimatorCoreEditorMenuOptions MenuOptions({EPropertyAnimatorCoreEditorMenuType::NewAdvanced});
 	EditorSubsystem->FillAnimatorMenu(InToolMenu, MenuContext, MenuOptions);
 }
 
diff --git a/Engine/Plugins/Experimental/PropertyAnimatorCore/Source/PropertyAnimatorCoreEditor/Public/Menus/PropertyAnimatorCoreEditorMenuDefs.h b/Engine/Plugins/Experimental/PropertyAnimatorCore/Source/PropertyAnimatorCoreEditor/Public/Menus/PropertyAnimatorCoreEditorMenuDefs.h
--- a/Engine/Plugins/Experimental/PropertyAnimatorCore/Source/PropertyAnimatorCoreEditor/Public/Menus/PropertyAnimatorCoreEditorMenuDefs.h
+++ b/Engine/Plugins/Experimental/PropertyAnimatorCore/Source/PropertyAnimatorCoreEditor/Public/Menus/PropertyAnimatorCoreEditorMenuDefs.h
@@ -18,17 +18,19 @@ enum class EPropertyAnimatorCoreEditorMenuType : uint8
 	/** Menu to open edit animators window */
 	Edit = 1 << 0,
 	/** Menu to add new animator on selected actors */
-	New = 1 << 1,
+	NewSimple = 1 << 1,
+	/** Menu to add new animator on selected actors, with submenus for presets and properties */
+	NewAdvanced = 1 << 2,
 	/** Menu to link or unlink properties/preset to/from existing animators */
-	Existing = 1 << 2,
+	Existing = 1 << 3,
 	/** Menu to link properties/preset to existing animator */
-	Link = 1 << 3,
+	Link = 1 << 4,
 	/** Menu to delete animators */
-	Delete = 1 << 4,
+	Delete = 1 << 5,
 	/** Menu to enable animators */
-	Enable = 1 << 5,
+	Enable = 1 << 6,
 	/** Menu to disable animators */
-	Disable = 1 << 6,
+	Disable = 1 << 7
 };
 
 /** Context for the menu builder, you can reuse it for multiple menus types */
