MotionDesign : PropertyAnimator
- Adding native support for FRotator and FVector in animator
- Displaying property type next to its name in the UI
- Improving animator API to evaluate and check property support
- Adapting menu to show supported parent and children properties

diff --git a/Engine/Plugins/Experimental/PropertyAnimator/Source/PropertyAnimator/Private/Animators/PropertyAnimatorBounce.cpp b/Engine/Plugins/Experimental/PropertyAnimator/Source/PropertyAnimator/Private/Animators/PropertyAnimatorBounce.cpp
--- a/Engine/Plugins/Experimental/PropertyAnimator/Source/PropertyAnimator/Private/Animators/PropertyAnimatorBounce.cpp
+++ b/Engine/Plugins/Experimental/PropertyAnimator/Source/PropertyAnimator/Private/Animators/PropertyAnimatorBounce.cpp
@@ -22,15 +22,19 @@ void UPropertyAnimatorBounce::SetInvertEffect(bool bInvert)
 	OnInvertEffect();
 }
 
-float UPropertyAnimatorBounce::Evaluate(double InTimeElapsed, const FPropertyAnimatorCoreData& InPropertyData, UPropertyAnimatorFloatContext* InOptions) const
+bool UPropertyAnimatorBounce::EvaluateProperty(const FPropertyAnimatorCoreData& InPropertyData, UPropertyAnimatorCoreContext* InContext, FInstancedPropertyBag& InParameters, FInstancedPropertyBag& OutEvaluationResult) const
 {
-	const float Frequency = InOptions->GetFrequency() * GlobalFrequency;
-	double TimeProgress = FMath::Fmod(InTimeElapsed + InOptions->GetTimeOffset(), 1.f / Frequency);
+	const float Frequency = InParameters.GetValueDouble(FrequencyParameterName).GetValue();
+	const double TimeElapsed = InParameters.GetValueDouble(TimeElapsedParameterName).GetValue();
+
+	double TimeProgress = FMath::Fmod(TimeElapsed, 1.f / Frequency);
 
 	TimeProgress = bInvertEffect ? TimeProgress : 1 - TimeProgress;
 
 	const double EasingValue = UE::PropertyAnimator::Easing::Bounce(TimeProgress, EPropertyAnimatorEasingType::In);
 
-	// Remap from [0, 1] to user amplitude from [Min, Max]
-	return FMath::GetMappedRangeValueClamped(FVector2D(0, 1), FVector2D(InOptions->GetAmplitudeMin(), InOptions->GetAmplitudeMax()), EasingValue);
+	InParameters.AddProperty(AlphaParameterName, EPropertyBagPropertyType::Float);
+	InParameters.SetValueFloat(AlphaParameterName, EasingValue);
+
+	return InContext->EvaluateProperty(InPropertyData, InParameters, OutEvaluationResult);
 }
diff --git a/Engine/Plugins/Experimental/PropertyAnimator/Source/PropertyAnimator/Private/Animators/PropertyAnimatorClock.cpp b/Engine/Plugins/Experimental/PropertyAnimator/Source/PropertyAnimator/Private/Animators/PropertyAnimatorClock.cpp
--- a/Engine/Plugins/Experimental/PropertyAnimator/Source/PropertyAnimator/Private/Animators/PropertyAnimatorClock.cpp
+++ b/Engine/Plugins/Experimental/PropertyAnimator/Source/PropertyAnimator/Private/Animators/PropertyAnimatorClock.cpp
@@ -2,7 +2,6 @@
 
 #include "Animators/PropertyAnimatorClock.h"
 
-#include "Internationalization/Regex.h"
 #include "Misc/DateTime.h"
 #include "Properties/Converters/PropertyAnimatorCoreConverterBase.h"
 #include "Properties/Handlers/PropertyAnimatorCoreHandlerBase.h"
@@ -11,22 +10,6 @@
 UPropertyAnimatorClock::UPropertyAnimatorClock()
 {
 	SetAnimatorDisplayName(DefaultControllerName);
-
-	if (!IsTemplate())
-	{
-		OnModeChanged();
-	}
-}
-
-void UPropertyAnimatorClock::SetMode(EPropertyAnimatorClockMode InMode)
-{
-	if (Mode == InMode)
-	{
-		return;
-	}
-
-	Mode = InMode;
-	OnModeChanged();
 }
 
 void UPropertyAnimatorClock::SetDisplayFormat(const FString& InDisplayFormat)
@@ -34,90 +17,35 @@ void UPropertyAnimatorClock::SetDisplayFormat(const FString& InDisplayFormat)
 	DisplayFormat = InDisplayFormat;
 }
 
-void UPropertyAnimatorClock::SetCountdownDuration(const FString& InDuration)
+EPropertyAnimatorPropertySupport UPropertyAnimatorClock::IsPropertySupported(const FPropertyAnimatorCoreData& InPropertyData) const
 {
-	if (CountdownDuration == InDuration)
+	if (InPropertyData.IsA<FStrProperty>())
 	{
-		return;
+		return EPropertyAnimatorPropertySupport::Complete;
 	}
 
-	CountdownDuration = InDuration;
-	OnModeChanged();
-}
-
-bool UPropertyAnimatorClock::IsPropertyDirectlySupported(const FPropertyAnimatorCoreData& InPropertyData) const
-{
-	return InPropertyData.IsA<FStrProperty>();
-}
-
-bool UPropertyAnimatorClock::IsPropertyIndirectlySupported(const FPropertyAnimatorCoreData& InPropertyData) const
-{
 	// Check if a converter supports the conversion
 	if (UPropertyAnimatorCoreSubsystem* AnimatorSubsystem = UPropertyAnimatorCoreSubsystem::Get())
 	{
 		static const FPropertyBagPropertyDesc AnimatorTypeDesc("", EPropertyBagPropertyType::String);
 		const FPropertyBagPropertyDesc PropertyTypeDesc("", InPropertyData.GetLeafProperty());
 
-		return AnimatorSubsystem->IsConversionSupported(AnimatorTypeDesc, PropertyTypeDesc);
+		if (AnimatorSubsystem->IsConversionSupported(AnimatorTypeDesc, PropertyTypeDesc))
+		{
+			return EPropertyAnimatorPropertySupport::Incomplete;
+		}
 	}
 
-	return false;
+	return Super::IsPropertySupported(InPropertyData);
 }
 
-void UPropertyAnimatorClock::PostLoad()
+void UPropertyAnimatorClock::EvaluateProperties(FInstancedPropertyBag& InParameters)
 {
-	Super::PostLoad();
-
-	OnModeChanged();
-}
-
-#if WITH_EDITOR
-void UPropertyAnimatorClock::PostEditChangeProperty(FPropertyChangedEvent& InPropertyChangedEvent)
-{
-	Super::PostEditChangeProperty(InPropertyChangedEvent);
-
-	const FName MemberName = InPropertyChangedEvent.GetMemberPropertyName();
+	const double TimeElapsed = InParameters.GetValueDouble(TimeElapsedParameterName).GetValue();
 
-	if (MemberName == GET_MEMBER_NAME_CHECKED(UPropertyAnimatorClock, Mode)
-		|| MemberName == GET_MEMBER_NAME_CHECKED(UPropertyAnimatorClock, CountdownDuration))
-	{
-		OnModeChanged();
-	}
-}
-#endif
-
-void UPropertyAnimatorClock::OnModeChanged()
-{
-	if (Mode == EPropertyAnimatorClockMode::LocalTime)
-	{
-		ActiveTimeSpan = FTimespan(FDateTime::Now().GetTicks()) - ElapsedTimeSpan;
-	}
-	else if (Mode == EPropertyAnimatorClockMode::Countdown)
-	{
-		ActiveTimeSpan = ParseTime(CountdownDuration);
-	}
-	else if (Mode == EPropertyAnimatorClockMode::Stopwatch)
-	{
-		ActiveTimeSpan = FTimespan::Zero();
-	}
-}
-
-void UPropertyAnimatorClock::EvaluateProperties(const FPropertyAnimatorCoreEvaluationParameters& InParameters)
-{
-	FString FormattedDateTime;
-
-	if (Mode == EPropertyAnimatorClockMode::Countdown)
-	{
-		ElapsedTimeSpan = ActiveTimeSpan - FTimespan::FromSeconds(InParameters.TimeElapsed);
-		const FDateTime CountdownDateTime(ElapsedTimeSpan > FTimespan::Zero() ? ElapsedTimeSpan.GetTicks() : 0);
-		FormattedDateTime = CountdownDateTime.ToFormattedString(*DisplayFormat);
-	}
-	else if (Mode == EPropertyAnimatorClockMode::Stopwatch || Mode == EPropertyAnimatorClockMode::LocalTime)
-	{
-		ElapsedTimeSpan = ActiveTimeSpan + FTimespan::FromSeconds(InParameters.TimeElapsed);
-		const FDateTime StopWatchDateTime(ElapsedTimeSpan.GetTicks());
-		FormattedDateTime = StopWatchDateTime.ToFormattedString(*DisplayFormat);
-	}
+	const FTimespan ElapsedTimeSpan = FTimespan::FromSeconds(TimeElapsed);
+	const FDateTime DateTime(ElapsedTimeSpan > FTimespan::Zero() ? ElapsedTimeSpan.GetTicks() : 0);
+	const FString FormattedDateTime = DateTime.ToFormattedString(*DisplayFormat);
 
 	EvaluateEachLinkedProperty<UPropertyAnimatorCoreContext>([this, FormattedDateTime](
 		UPropertyAnimatorCoreContext* InContext,
@@ -133,9 +61,9 @@ void UPropertyAnimatorClock::EvaluateProperties(const FPropertyAnimatorCoreEvalu
 	});
 }
 
-void UPropertyAnimatorClock::OnPropertyLinked(UPropertyAnimatorCoreContext* InLinkedProperty)
+void UPropertyAnimatorClock::OnPropertyLinked(UPropertyAnimatorCoreContext* InLinkedProperty, EPropertyAnimatorPropertySupport InSupport)
 {
-	Super::OnPropertyLinked(InLinkedProperty);
+	Super::OnPropertyLinked(InLinkedProperty, InSupport);
 
 	const FPropertyAnimatorCoreData& Property = InLinkedProperty->GetAnimatedProperty();
 	if (Property.IsA<FStrProperty>())
@@ -143,49 +71,15 @@ void UPropertyAnimatorClock::OnPropertyLinked(UPropertyAnimatorCoreContext* InLi
 		return;
 	}
 
-	const UPropertyAnimatorCoreSubsystem* AnimatorSubsystem = UPropertyAnimatorCoreSubsystem::Get();
-	if (!AnimatorSubsystem)
-	{
-		return;
-	}
-
-	static const FPropertyBagPropertyDesc AnimatorTypeDesc("", EPropertyBagPropertyType::String);
-	const FPropertyBagPropertyDesc PropertyTypeDesc("", Property.GetLeafProperty());
-	const TSet<UPropertyAnimatorCoreConverterBase*> Converters = AnimatorSubsystem->GetSupportedConverters(AnimatorTypeDesc, PropertyTypeDesc);
-	check(!Converters.IsEmpty())
-	InLinkedProperty->SetConverterClass(Converters.Array()[0]->GetClass());
-}
-
-FTimespan UPropertyAnimatorClock::ParseTime(const FString& InFormat)
-{
-	// Regex patterns for different formats
-	static const FRegexPattern HHMMSSPattern(TEXT("^(?:(\\d{2}):)?(\\d{2}):(\\d{2})$")); // 01:00 00:01:00
-	static const FRegexPattern CombinedPattern(TEXT("(?:(\\d+)h)? ?(?:(\\d+)m)? ?(?:(\\d+)s)?")); // 1h 1m 1s
-
-	FRegexMatcher HHMMSSMatcher(HHMMSSPattern, InFormat);
-	FRegexMatcher CombinedMatcher(CombinedPattern, InFormat);
-
-	FTimespan ParsedTimeSpan = FTimespan::Zero();
-
-	if (InFormat.IsNumeric())
-	{
-		const int32 Seconds = FCString::Atoi(*InFormat);
-		ParsedTimeSpan = FTimespan::FromSeconds(Seconds);
-	}
-	else if (HHMMSSMatcher.FindNext())
+	if (EnumHasAnyFlags(InSupport, EPropertyAnimatorPropertySupport::Incomplete))
 	{
-		const int32 Hours = HHMMSSMatcher.GetCaptureGroup(1).IsEmpty() ? 0 : FCString::Atoi(*HHMMSSMatcher.GetCaptureGroup(1));
-		const int32 Minutes = FCString::Atoi(*HHMMSSMatcher.GetCaptureGroup(2));
-		const int32 Seconds = FCString::Atoi(*HHMMSSMatcher.GetCaptureGroup(3));
-		ParsedTimeSpan = FTimespan::FromHours(Hours) + FTimespan::FromMinutes(Minutes) + FTimespan::FromSeconds(Seconds);
+		if (const UPropertyAnimatorCoreSubsystem* AnimatorSubsystem = UPropertyAnimatorCoreSubsystem::Get())
+		{
+			static const FPropertyBagPropertyDesc AnimatorTypeDesc("", EPropertyBagPropertyType::String);
+			const FPropertyBagPropertyDesc PropertyTypeDesc("", Property.GetLeafProperty());
+			const TSet<UPropertyAnimatorCoreConverterBase*> Converters = AnimatorSubsystem->GetSupportedConverters(AnimatorTypeDesc, PropertyTypeDesc);
+			check(!Converters.IsEmpty())
+			InLinkedProperty->SetConverterClass(Converters.Array()[0]->GetClass());
+		}
 	}
-	else if (CombinedMatcher.FindNext())
-	{
-		const int32 Hours = CombinedMatcher.GetCaptureGroup(1).IsEmpty() ? 0 : FCString::Atoi(*CombinedMatcher.GetCaptureGroup(1));
-		const int32 Minutes = CombinedMatcher.GetCaptureGroup(2).IsEmpty() ? 0 : FCString::Atoi(*CombinedMatcher.GetCaptureGroup(2));
-		const int32 Seconds = CombinedMatcher.GetCaptureGroup(3).IsEmpty() ? 0 : FCString::Atoi(*CombinedMatcher.GetCaptureGroup(3));
-		ParsedTimeSpan = FTimespan::FromHours(Hours) + FTimespan::FromMinutes(Minutes) + FTimespan::FromSeconds(Seconds);
-	}
-
-	return ParsedTimeSpan;
-}
+}
\ No newline at end of file
diff --git a/Engine/Plugins/Experimental/PropertyAnimator/Source/PropertyAnimator/Private/Animators/PropertyAnimatorFloatBase.cpp b/Engine/Plugins/Experimental/PropertyAnimator/Source/PropertyAnimator/Private/Animators/PropertyAnimatorFloatBase.cpp
--- a/Engine/Plugins/Experimental/PropertyAnimator/Source/PropertyAnimator/Private/Animators/PropertyAnimatorFloatBase.cpp
+++ b/Engine/Plugins/Experimental/PropertyAnimator/Source/PropertyAnimator/Private/Animators/PropertyAnimatorFloatBase.cpp
@@ -3,6 +3,8 @@
 #include "Animators/PropertyAnimatorFloatBase.h"
 
 #include "Properties/PropertyAnimatorFloatContext.h"
+#include "Properties/PropertyAnimatorRotatorContext.h"
+#include "Properties/PropertyAnimatorVectorContext.h"
 #include "Properties/Converters/PropertyAnimatorCoreConverterBase.h"
 #include "Properties/Handlers/PropertyAnimatorCoreHandlerBase.h"
 #include "Subsystems/PropertyAnimatorCoreSubsystem.h"
@@ -25,56 +27,48 @@ void UPropertyAnimatorFloatBase::PostEditChangeProperty(FPropertyChangedEvent& I
 }
 #endif
 
-bool UPropertyAnimatorFloatBase::IsPropertyDirectlySupported(const FPropertyAnimatorCoreData& InPropertyData) const
+void UPropertyAnimatorFloatBase::SetMagnitude(float InMagnitude)
 {
-	return InPropertyData.IsA<FFloatProperty>();
-}
-
-bool UPropertyAnimatorFloatBase::IsPropertyIndirectlySupported(const FPropertyAnimatorCoreData& InPropertyData) const
-{
-	// Check if a converter supports the conversion
-	if (UPropertyAnimatorCoreSubsystem* AnimatorSubsystem = UPropertyAnimatorCoreSubsystem::Get())
+	if (FMath::IsNearlyEqual(Magnitude, InMagnitude))
 	{
-		static const FPropertyBagPropertyDesc AnimatorTypeDesc("", EPropertyBagPropertyType::Float);
-		const FPropertyBagPropertyDesc PropertyTypeDesc("", InPropertyData.GetLeafProperty());
-
-		return AnimatorSubsystem->IsConversionSupported(AnimatorTypeDesc, PropertyTypeDesc);
+		return;
 	}
 
-	return false;
+	Magnitude = InMagnitude;
+	OnMagnitudeChanged();
 }
 
-void UPropertyAnimatorFloatBase::SetGlobalMagnitude(float InMagnitude)
+void UPropertyAnimatorFloatBase::SetCycleDuration(float InCycleDuration)
 {
-	if (FMath::IsNearlyEqual(GlobalMagnitude, InMagnitude))
+	if (FMath::IsNearlyEqual(CycleDuration, InCycleDuration))
 	{
 		return;
 	}
 
-	GlobalMagnitude = InMagnitude;
-	OnMagnitudeChanged();
+	CycleDuration = InCycleDuration;
+	OnCycleDurationChanged();
 }
 
-void UPropertyAnimatorFloatBase::SetGlobalFrequency(float InFrequency)
+void UPropertyAnimatorFloatBase::SetCycleMode(EPropertyAnimatorCycleMode InMode)
 {
-	if (FMath::IsNearlyEqual(GlobalFrequency, InFrequency))
+	if (CycleMode == InMode)
 	{
 		return;
 	}
 
-	GlobalFrequency = InFrequency;
-	OnGlobalFrequencyChanged();
+	CycleMode = InMode;
+	OnCycleModeChanged();
 }
 
-void UPropertyAnimatorFloatBase::SetAccumulatedTimeOffset(double InOffset)
+void UPropertyAnimatorFloatBase::SetTimeOffset(double InOffset)
 {
-	if (FMath::IsNearlyEqual(AccumulatedTimeOffset, InOffset))
+	if (FMath::IsNearlyEqual(TimeOffset, InOffset))
 	{
 		return;
 	}
 
-	AccumulatedTimeOffset = InOffset;
-	OnAccumulatedTimeOffsetChanged();
+	TimeOffset = InOffset;
+	OnTimeOffsetChanged();
 }
 
 void UPropertyAnimatorFloatBase::SetRandomTimeOffset(bool bInOffset)
@@ -99,65 +93,150 @@ void UPropertyAnimatorFloatBase::SetSeed(int32 InSeed)
 	OnSeedChanged();
 }
 
-TSubclassOf<UPropertyAnimatorCoreContext> UPropertyAnimatorFloatBase::GetPropertyContextClass(const FPropertyAnimatorCoreData& InUnlinkedProperty)
+TSubclassOf<UPropertyAnimatorCoreContext> UPropertyAnimatorFloatBase::GetPropertyContextClass(const FPropertyAnimatorCoreData& InProperty)
 {
+	if (InProperty.IsA<FStructProperty>())
+	{
+		const FName TypeName = InProperty.GetLeafPropertyTypeName();
+
+		if (TypeName == NAME_Rotator)
+		{
+			return UPropertyAnimatorRotatorContext::StaticClass();
+		}
+
+		if (TypeName == NAME_Vector)
+		{
+			return UPropertyAnimatorVectorContext::StaticClass();
+		}
+	}
+
 	return UPropertyAnimatorFloatContext::StaticClass();
 }
 
-void UPropertyAnimatorFloatBase::EvaluateProperties(const FPropertyAnimatorCoreEvaluationParameters& InParameters)
+EPropertyAnimatorPropertySupport UPropertyAnimatorFloatBase::IsPropertySupported(const FPropertyAnimatorCoreData& InPropertyData) const
+{
+	const FName TypeName = InPropertyData.GetLeafPropertyTypeName();
+
+	if (InPropertyData.IsA<FFloatProperty>())
+	{
+		return EPropertyAnimatorPropertySupport::Complete;
+	}
+
+	if (InPropertyData.IsA<FStructProperty>())
+	{
+		if (TypeName == NAME_Rotator)
+		{
+			return EPropertyAnimatorPropertySupport::Complete;
+		}
+
+		if (TypeName == NAME_Vector)
+		{
+			return EPropertyAnimatorPropertySupport::Complete;
+		}
+	}
+
+	// Check if a converter supports the conversion
+	if (UPropertyAnimatorCoreSubsystem* AnimatorSubsystem = UPropertyAnimatorCoreSubsystem::Get())
+	{
+		static const FPropertyBagPropertyDesc AnimatorTypeDesc("", EPropertyBagPropertyType::Float);
+		const FPropertyBagPropertyDesc PropertyTypeDesc("", InPropertyData.GetLeafProperty());
+
+		if (AnimatorSubsystem->IsConversionSupported(AnimatorTypeDesc, PropertyTypeDesc))
+		{
+			return EPropertyAnimatorPropertySupport::Incomplete;
+		}
+	}
+
+	return Super::IsPropertySupported(InPropertyData);
+}
+
+void UPropertyAnimatorFloatBase::EvaluateProperties(FInstancedPropertyBag& InParameters)
 {
-	double AccumulatedTimeOffsetRound = InParameters.TimeElapsed;
-	const float AnimatorMagnitude = GlobalMagnitude * InParameters.AnimatorsMagnitude;
+	const float AnimatorMagnitude = Magnitude * InParameters.GetValueFloat(MagnitudeParameterName).GetValue();
+	double TimeElapsed = InParameters.GetValueDouble(TimeElapsedParameterName).GetValue();
 	RandomStream = FRandomStream(Seed);
 
-	EvaluateEachLinkedProperty<UPropertyAnimatorFloatContext>([this, &AccumulatedTimeOffsetRound, &AnimatorMagnitude](
-		UPropertyAnimatorFloatContext* InOptions
+	if (CycleMode == EPropertyAnimatorCycleMode::DoOnce)
+	{
+		if (FMath::Abs(TimeElapsed) > CycleDuration)
+		{
+			return;
+		}
+	}
+	else if (CycleMode == EPropertyAnimatorCycleMode::Loop)
+	{
+		TimeElapsed = FMath::Fmod(TimeElapsed, CycleDuration + CycleGapDuration);
+
+		if (TimeElapsed > CycleDuration)
+		{
+			TimeElapsed = CycleDuration - UE_KINDA_SMALL_NUMBER;
+		}
+	}
+	else if (CycleMode == EPropertyAnimatorCycleMode::PingPong)
+	{
+		const bool bReverse = FMath::Modulo(FMath::TruncToInt32(TimeElapsed / (CycleDuration + CycleGapDuration)), 2) != 0;
+		TimeElapsed = FMath::Fmod(TimeElapsed, CycleDuration + CycleGapDuration);
+
+		if (TimeElapsed > CycleDuration)
+		{
+			TimeElapsed = CycleDuration - UE_KINDA_SMALL_NUMBER;
+		}
+
+		if (bReverse)
+		{
+			TimeElapsed = CycleDuration - FMath::Fmod(TimeElapsed, CycleDuration);
+		}
+		else
+		{
+			TimeElapsed = FMath::Fmod(TimeElapsed, CycleDuration);
+		}
+	}
+
+	EvaluateEachLinkedProperty<UPropertyAnimatorCoreContext>([this, &TimeElapsed, &AnimatorMagnitude, &InParameters](
+		UPropertyAnimatorCoreContext* InOptions
 		, const FPropertyAnimatorCoreData& InResolvedProperty
 		, FInstancedPropertyBag& InEvaluatedValues)->bool
 	{
 		const double RandomTimeOffset = bRandomTimeOffset ? RandomStream.GetFraction() : 0;
-		AccumulatedTimeOffsetRound += AccumulatedTimeOffset + RandomTimeOffset;
+		TimeElapsed += TimeOffset + RandomTimeOffset;
 
-		if (GlobalMagnitude != 0
-			&& GlobalFrequency != 0
-			&& InOptions->GetMagnitude() != 0
-			&& InOptions->GetFrequency() != 0)
+		if (Magnitude != 0
+			&& CycleDuration > 0
+			&& InOptions->GetMagnitude() != 0)
 		{
-			const float EvaluationResult = AnimatorMagnitude
-				* InOptions->GetMagnitude()
-				* Evaluate(AccumulatedTimeOffsetRound, InResolvedProperty, InOptions);
+			// Frequency
+			InParameters.AddProperty(FrequencyParameterName, EPropertyBagPropertyType::Float);
+			InParameters.SetValueFloat(FrequencyParameterName, 1.f / CycleDuration);
 
-			const FName DisplayName(InResolvedProperty.GetPathHash());
+			// Time Elapsed
+			InParameters.SetValueDouble(TimeElapsedParameterName, TimeElapsed + InOptions->GetTimeOffset());
 
-			InEvaluatedValues.AddProperty(DisplayName, EPropertyBagPropertyType::Float);
-			InEvaluatedValues.SetValueFloat(DisplayName, EvaluationResult);
+			// Magnitude
+			InParameters.SetValueFloat(MagnitudeParameterName, AnimatorMagnitude * InOptions->GetMagnitude());
 
-			return true;
+			return EvaluateProperty(InResolvedProperty, InOptions, InParameters, InEvaluatedValues);
 		}
 
 		return false;
 	});
 }
 
-void UPropertyAnimatorFloatBase::OnPropertyLinked(UPropertyAnimatorCoreContext* InLinkedProperty)
+void UPropertyAnimatorFloatBase::OnPropertyLinked(UPropertyAnimatorCoreContext* InLinkedProperty, EPropertyAnimatorPropertySupport InSupport)
 {
-	Super::OnPropertyLinked(InLinkedProperty);
+	Super::OnPropertyLinked(InLinkedProperty, InSupport);
 
-	const FPropertyAnimatorCoreData& Property = InLinkedProperty->GetAnimatedProperty();
-	if (Property.IsA<FFloatProperty>())
+	if (EnumHasAnyFlags(InSupport, EPropertyAnimatorPropertySupport::Incomplete))
 	{
-		return;
-	}
+		if (const UPropertyAnimatorCoreSubsystem* AnimatorSubsystem = UPropertyAnimatorCoreSubsystem::Get())
+		{
+			static const FPropertyBagPropertyDesc AnimatorTypeDesc("", EPropertyBagPropertyType::Float);
+			const FPropertyBagPropertyDesc PropertyTypeDesc("", InLinkedProperty->GetAnimatedProperty().GetLeafProperty());
 
-	const UPropertyAnimatorCoreSubsystem* AnimatorSubsystem = UPropertyAnimatorCoreSubsystem::Get();
-	if (!AnimatorSubsystem)
-	{
-		return;
-	}
+			const TSet<UPropertyAnimatorCoreConverterBase*> Converters = AnimatorSubsystem->GetSupportedConverters(AnimatorTypeDesc, PropertyTypeDesc);
+
+			check(!Converters.IsEmpty())
 
-	static const FPropertyBagPropertyDesc AnimatorTypeDesc("", EPropertyBagPropertyType::Float);
-	const FPropertyBagPropertyDesc PropertyTypeDesc("", Property.GetLeafProperty());
-	const TSet<UPropertyAnimatorCoreConverterBase*> Converters = AnimatorSubsystem->GetSupportedConverters(AnimatorTypeDesc, PropertyTypeDesc);
-	check(!Converters.IsEmpty())
-	InLinkedProperty->SetConverterClass(Converters.Array()[0]->GetClass());
+			InLinkedProperty->SetConverterClass(Converters.Array()[0]->GetClass());
+		}
+	}
 }
diff --git a/Engine/Plugins/Experimental/PropertyAnimator/Source/PropertyAnimator/Private/Animators/PropertyAnimatorOscillate.cpp b/Engine/Plugins/Experimental/PropertyAnimator/Source/PropertyAnimator/Private/Animators/PropertyAnimatorOscillate.cpp
--- a/Engine/Plugins/Experimental/PropertyAnimator/Source/PropertyAnimator/Private/Animators/PropertyAnimatorOscillate.cpp
+++ b/Engine/Plugins/Experimental/PropertyAnimator/Source/PropertyAnimator/Private/Animators/PropertyAnimatorOscillate.cpp
@@ -15,38 +15,46 @@ void UPropertyAnimatorOscillate::SetOscillateFunction(EPropertyAnimatorOscillate
 	OscillateFunction = InFunction;
 }
 
-float UPropertyAnimatorOscillate::Evaluate(double InTimeElapsed, const FPropertyAnimatorCoreData& InPropertyData, UPropertyAnimatorFloatContext* InOptions) const
+bool UPropertyAnimatorOscillate::EvaluateProperty(const FPropertyAnimatorCoreData& InPropertyData, UPropertyAnimatorCoreContext* InContext, FInstancedPropertyBag& InParameters, FInstancedPropertyBag& OutEvaluationResult) const
 {
 	double WaveResult = 0.f;
 
 	using namespace UE::PropertyAnimator;
 
-	const float Frequency = InOptions->GetFrequency() * GlobalFrequency;
+	const float Frequency = InParameters.GetValueDouble(FrequencyParameterName).GetValue();
+	const double TimeElapsed = InParameters.GetValueDouble(TimeElapsedParameterName).GetValue();
+
+	constexpr double Amplitude = 1;
+	constexpr double Offset = 0;
 
 	switch (OscillateFunction)
 	{
 		case EPropertyAnimatorOscillateFunction::Sine:
-			WaveResult = Wave::Sine(InTimeElapsed, 1.f, Frequency, InOptions->GetTimeOffset());
-		break;
+			WaveResult = Wave::Sine(TimeElapsed, Amplitude, Frequency, Offset);
+			break;
 		case EPropertyAnimatorOscillateFunction::Cosine:
-			WaveResult = Wave::Cosine(InTimeElapsed, 1.f, Frequency, InOptions->GetTimeOffset());
-		break;
+			WaveResult = Wave::Cosine(TimeElapsed, Amplitude, Frequency, Offset);
+			break;
 		case EPropertyAnimatorOscillateFunction::Square:
-			WaveResult = Wave::Square(InTimeElapsed, 1.f, Frequency, InOptions->GetTimeOffset());
-		break;
+			WaveResult = Wave::Square(TimeElapsed, Amplitude, Frequency, Offset);
+			break;
 		case EPropertyAnimatorOscillateFunction::InvertedSquare:
-			WaveResult = Wave::InvertedSquare(InTimeElapsed, 1.f, Frequency, InOptions->GetTimeOffset());
-		break;
+			WaveResult = Wave::InvertedSquare(TimeElapsed, Amplitude, Frequency, Offset);
+			break;
 		case EPropertyAnimatorOscillateFunction::Sawtooth:
-			WaveResult = Wave::Sawtooth(InTimeElapsed, 1.f, Frequency, InOptions->GetTimeOffset());
-		break;
+			WaveResult = Wave::Sawtooth(TimeElapsed, Amplitude, Frequency, Offset);
+			break;
 		case EPropertyAnimatorOscillateFunction::Triangle:
-			WaveResult = Wave::Triangle(InTimeElapsed, 1.f, Frequency, InOptions->GetTimeOffset());
-		break;
+			WaveResult = Wave::Triangle(TimeElapsed, Amplitude, Frequency, Offset);
+			break;
 		default:
 			checkNoEntry();
 	}
 
-	// Remap from [-1, 1] to user amplitude from [Min, Max]
-	return FMath::GetMappedRangeValueClamped(FVector2D(-1, 1), FVector2D(InOptions->GetAmplitudeMin(), InOptions->GetAmplitudeMax()), WaveResult);
+	const double NormalizedResult = FMath::GetMappedRangeValueClamped(FVector2D(-1, 1), FVector2D(0, 1), WaveResult);
+
+	InParameters.AddProperty(AlphaParameterName, EPropertyBagPropertyType::Float);
+	InParameters.SetValueFloat(AlphaParameterName, NormalizedResult);
+
+	return InContext->EvaluateProperty(InPropertyData, InParameters, OutEvaluationResult);
 }
diff --git a/Engine/Plugins/Experimental/PropertyAnimator/Source/PropertyAnimator/Private/Animators/PropertyAnimatorPulse.cpp b/Engine/Plugins/Experimental/PropertyAnimator/Source/PropertyAnimator/Private/Animators/PropertyAnimatorPulse.cpp
--- a/Engine/Plugins/Experimental/PropertyAnimator/Source/PropertyAnimator/Private/Animators/PropertyAnimatorPulse.cpp
+++ b/Engine/Plugins/Experimental/PropertyAnimator/Source/PropertyAnimator/Private/Animators/PropertyAnimatorPulse.cpp
@@ -20,11 +20,15 @@ void UPropertyAnimatorPulse::SetEasingType(EPropertyAnimatorEasingType InEasingT
 	EasingType = InEasingType;
 }
 
-float UPropertyAnimatorPulse::Evaluate(double InTimeElapsed, const FPropertyAnimatorCoreData& InPropertyData, UPropertyAnimatorFloatContext* InOptions) const
+bool UPropertyAnimatorPulse::EvaluateProperty(const FPropertyAnimatorCoreData& InPropertyData, UPropertyAnimatorCoreContext* InContext, FInstancedPropertyBag& InParameters, FInstancedPropertyBag& OutEvaluationResult) const
 {
-	const float Frequency = InOptions->GetFrequency() * GlobalFrequency;
+	const float Frequency = InParameters.GetValueDouble(FrequencyParameterName).GetValue();
+	const double TimeElapsed = InParameters.GetValueDouble(TimeElapsedParameterName).GetValue();
 
-	const double WaveResult = UE::PropertyAnimator::Wave::Triangle(InTimeElapsed, 1.f, Frequency, InOptions->GetTimeOffset());
+	constexpr double Amplitude = 1;
+	constexpr double Offset = 0;
+
+	const double WaveResult = UE::PropertyAnimator::Wave::Triangle(TimeElapsed, Amplitude, Frequency, Offset);
 
 	// Result of wave functions is [-1, 1] -> remap to [0, 1] range for easing functions
 	const float NormalizedWaveProgress = FMath::GetMappedRangeValueClamped(FVector2D(-1, 1), FVector2D(0, 1), WaveResult);
@@ -32,6 +36,8 @@ float UPropertyAnimatorPulse::Evaluate(double InTimeElapsed, const FPropertyAnim
 	// Apply easing function on normalized progress
 	const float EasingResult = UE::PropertyAnimator::Easing::Ease(NormalizedWaveProgress, EasingFunction, EasingType);
 
-	// Remap from [0, 1] to user amplitude from [Min, Max]
-	return FMath::GetMappedRangeValueClamped(FVector2D(0, 1), FVector2D(InOptions->GetAmplitudeMin(), InOptions->GetAmplitudeMax()), EasingResult);
+	InParameters.AddProperty(AlphaParameterName, EPropertyBagPropertyType::Float);
+	InParameters.SetValueFloat(AlphaParameterName, EasingResult);
+
+	return InContext->EvaluateProperty(InPropertyData, InParameters, OutEvaluationResult);
 }
diff --git a/Engine/Plugins/Experimental/PropertyAnimator/Source/PropertyAnimator/Private/Animators/PropertyAnimatorSoundWave.cpp b/Engine/Plugins/Experimental/PropertyAnimator/Source/PropertyAnimator/Private/Animators/PropertyAnimatorSoundWave.cpp
--- a/Engine/Plugins/Experimental/PropertyAnimator/Source/PropertyAnimator/Private/Animators/PropertyAnimatorSoundWave.cpp
+++ b/Engine/Plugins/Experimental/PropertyAnimator/Source/PropertyAnimator/Private/Animators/PropertyAnimatorSoundWave.cpp
@@ -59,28 +59,34 @@ void UPropertyAnimatorSoundWave::OnSampledSoundWaveChanged()
 #endif
 }
 
-float UPropertyAnimatorSoundWave::Evaluate(double InTimeElapsed, const FPropertyAnimatorCoreData& InPropertyData, UPropertyAnimatorFloatContext* InOptions) const
+bool UPropertyAnimatorSoundWave::EvaluateProperty(const FPropertyAnimatorCoreData& InPropertyData, UPropertyAnimatorCoreContext* InContext, FInstancedPropertyBag& InParameters, FInstancedPropertyBag& OutEvaluationResult) const
 {
-	float Loudness = 0.f;
-
-	if (AudioAnalyzer
-		&& AudioAnalyzer->DurationInSeconds > 0)
+	if (AudioAnalyzer && AudioAnalyzer->DurationInSeconds > 0)
 	{
-		float SampleTime = InTimeElapsed + InOptions->GetTimeOffset();
+		double TimeElapsed = InParameters.GetValueDouble(TimeElapsedParameterName).GetValue();
+		float Frequency = InParameters.GetValueFloat(FrequencyParameterName).GetValue();
 
-		if ((SampleTime >= 0 && SampleTime <= AudioAnalyzer->DurationInSeconds) || bLoop)
-		{
-			SampleTime = FMath::Fmod(SampleTime, AudioAnalyzer->DurationInSeconds);
+		const float Period = 1.f / Frequency;
+		float SampleTime = FMath::Fmod(TimeElapsed, Period);
 
+		if (FMath::Abs(TimeElapsed / Period) <= 1.f || bLoop)
+		{
 			if (SampleTime < 0)
 			{
-				SampleTime = AudioAnalyzer->DurationInSeconds + SampleTime;
+				SampleTime = Period - FMath::Abs(SampleTime);
 			}
 
-			AudioAnalyzer->GetNormalizedLoudnessAtTime(SampleTime, Loudness);
+			const float NormalizedSampleTime = FMath::GetMappedRangeValueClamped(FVector2D(0, Period), FVector2D(0, AudioAnalyzer->DurationInSeconds), SampleTime);
+
+			float NormalizedLoudness = 0.f;
+			AudioAnalyzer->GetNormalizedLoudnessAtTime(NormalizedSampleTime, NormalizedLoudness);
+
+			InParameters.AddProperty(AlphaParameterName, EPropertyBagPropertyType::Float);
+			InParameters.SetValueFloat(AlphaParameterName, NormalizedLoudness);
+
+			return InContext->EvaluateProperty(InPropertyData, InParameters, OutEvaluationResult);
 		}
 	}
 
-	// Remap from [0, 1] to user amplitude from [Min, Max]
-	return FMath::GetMappedRangeValueClamped(FVector2D(0, 1), FVector2D(InOptions->GetAmplitudeMin(), InOptions->GetAmplitudeMax()), Loudness);
+	return false;
 }
diff --git a/Engine/Plugins/Experimental/PropertyAnimator/Source/PropertyAnimator/Private/Animators/PropertyAnimatorTime.cpp b/Engine/Plugins/Experimental/PropertyAnimator/Source/PropertyAnimator/Private/Animators/PropertyAnimatorTime.cpp
--- a/Engine/Plugins/Experimental/PropertyAnimator/Source/PropertyAnimator/Private/Animators/PropertyAnimatorTime.cpp
+++ b/Engine/Plugins/Experimental/PropertyAnimator/Source/PropertyAnimator/Private/Animators/PropertyAnimatorTime.cpp
@@ -10,13 +10,16 @@ UPropertyAnimatorTime::UPropertyAnimatorTime()
 	SetAnimatorDisplayName(DefaultControllerName);
 }
 
-float UPropertyAnimatorTime::Evaluate(double InTimeElapsed, const FPropertyAnimatorCoreData& InPropertyData, UPropertyAnimatorFloatContext* InOptions) const
+bool UPropertyAnimatorTime::EvaluateProperty(const FPropertyAnimatorCoreData& InPropertyData, UPropertyAnimatorCoreContext* InContext, FInstancedPropertyBag& InParameters, FInstancedPropertyBag& OutEvaluationResult) const
 {
-	const float Frequency = InOptions->GetFrequency() * GlobalFrequency;
-
+	const double TimeElapsed = InParameters.GetValueDouble(TimeElapsedParameterName).GetValue();
+	const double Frequency = InParameters.GetValueDouble(FrequencyParameterName).GetValue();
 	const double TimePeriod = 1.f / Frequency;
+	const double TimeProgress = FMath::Fmod(TimeElapsed, TimePeriod);
+	const float NormalizedValue = FMath::GetMappedRangeValueClamped(FVector2f(0, TimePeriod), FVector2f(0, 1), TimeProgress);
 
-	const double TimeProgress = FMath::Fmod(InTimeElapsed + InOptions->GetTimeOffset(), TimePeriod);
+	InParameters.AddProperty(AlphaParameterName, EPropertyBagPropertyType::Float);
+	InParameters.SetValueFloat(AlphaParameterName, NormalizedValue);
 
-	return FMath::GetMappedRangeValueClamped(FVector2f(0, TimePeriod), FVector2f(InOptions->GetAmplitudeMin(), InOptions->GetAmplitudeMax()), TimeProgress);
+	return InContext->EvaluateProperty(InPropertyData, InParameters, OutEvaluationResult);
 }
diff --git a/Engine/Plugins/Experimental/PropertyAnimator/Source/PropertyAnimator/Private/Animators/PropertyAnimatorWiggle.cpp b/Engine/Plugins/Experimental/PropertyAnimator/Source/PropertyAnimator/Private/Animators/PropertyAnimatorWiggle.cpp
--- a/Engine/Plugins/Experimental/PropertyAnimator/Source/PropertyAnimator/Private/Animators/PropertyAnimatorWiggle.cpp
+++ b/Engine/Plugins/Experimental/PropertyAnimator/Source/PropertyAnimator/Private/Animators/PropertyAnimatorWiggle.cpp
@@ -16,13 +16,19 @@ UPropertyAnimatorWiggle::UPropertyAnimatorWiggle()
 	Seed = SeedIncrement++;
 }
 
-float UPropertyAnimatorWiggle::Evaluate(double InTimeElapsed, const FPropertyAnimatorCoreData& InPropertyData, UPropertyAnimatorFloatContext* InOptions) const
+bool UPropertyAnimatorWiggle::EvaluateProperty(const FPropertyAnimatorCoreData& InPropertyData, UPropertyAnimatorCoreContext* InContext, FInstancedPropertyBag& InParameters, FInstancedPropertyBag& OutEvaluationResult) const
 {
-	const float Frequency = InOptions->GetFrequency() * GlobalFrequency;
+	const double TimeElapsed = InParameters.GetValueDouble(TimeElapsedParameterName).GetValue();
+	const double Frequency = InParameters.GetValueDouble(FrequencyParameterName).GetValue();
 
 	// Apply random wave based on time and frequency
-	const double WaveResult = UE::PropertyAnimator::Wave::Perlin(InTimeElapsed, 1.f, Frequency, InOptions->GetTimeOffset());
+	const double WaveResult = UE::PropertyAnimator::Wave::Perlin(TimeElapsed, 1.f, Frequency, 0.f);
 
 	// Remap from [-1, 1] to user amplitude from [Min, Max]
-	return FMath::GetMappedRangeValueClamped(FVector2D(-1, 1), FVector2D(InOptions->GetAmplitudeMin(), InOptions->GetAmplitudeMax()), WaveResult);
+	const float NormalizedValue = FMath::GetMappedRangeValueClamped(FVector2D(-1, 1), FVector2D(0, 1), WaveResult);
+
+	InParameters.AddProperty(AlphaParameterName, EPropertyBagPropertyType::Float);
+	InParameters.SetValueFloat(AlphaParameterName, NormalizedValue);
+
+	return InContext->EvaluateProperty(InPropertyData, InParameters, OutEvaluationResult);
 }
diff --git a/Engine/Plugins/Experimental/PropertyAnimator/Source/PropertyAnimator/Private/Properties/PropertyAnimatorFloatContext.cpp b/Engine/Plugins/Experimental/PropertyAnimator/Source/PropertyAnimator/Private/Properties/PropertyAnimatorFloatContext.cpp
--- a/Engine/Plugins/Experimental/PropertyAnimator/Source/PropertyAnimator/Private/Properties/PropertyAnimatorFloatContext.cpp
+++ b/Engine/Plugins/Experimental/PropertyAnimator/Source/PropertyAnimator/Private/Properties/PropertyAnimatorFloatContext.cpp
@@ -2,38 +2,61 @@
 
 #include "Properties/PropertyAnimatorFloatContext.h"
 
-void UPropertyAnimatorFloatContext::SetMagnitude(float InMagnitude)
-{
-	Magnitude = FMath::Clamp(InMagnitude, 0.f, 1.f);
-}
+#include "Animators/PropertyAnimatorCoreBase.h"
 
 void UPropertyAnimatorFloatContext::SetAmplitudeMin(double InAmplitude)
 {
-	AmplitudeMin = InAmplitude;
+	AmplitudeMin = GetClampedAmplitude(InAmplitude);
 }
 
 void UPropertyAnimatorFloatContext::SetAmplitudeMax(double InAmplitude)
 {
-	AmplitudeMax = InAmplitude;
+	AmplitudeMax = GetClampedAmplitude(InAmplitude);
 }
 
-void UPropertyAnimatorFloatContext::SetFrequency(float InFrequency)
+#if WITH_EDITOR
+void UPropertyAnimatorFloatContext::PostEditChangeProperty(FPropertyChangedEvent& InPropertyChangedEvent)
 {
-	Frequency = FMath::Max(0.f, InFrequency);
+	Super::PostEditChangeProperty(InPropertyChangedEvent);
+
+	const FName MemberPropertyName = InPropertyChangedEvent.GetMemberPropertyName();
+
+	if (MemberPropertyName == GET_MEMBER_NAME_CHECKED(UPropertyAnimatorFloatContext, AmplitudeMin))
+	{
+		SetAmplitudeMin(AmplitudeMin);
+	}
+	else if (MemberPropertyName == GET_MEMBER_NAME_CHECKED(UPropertyAnimatorFloatContext, AmplitudeMax))
+	{
+		SetAmplitudeMax(AmplitudeMax);
+	}
 }
+#endif
 
-void UPropertyAnimatorFloatContext::SetTimeOffset(double InTimeOffset)
+bool UPropertyAnimatorFloatContext::EvaluateProperty(const FPropertyAnimatorCoreData& InProperty, const FInstancedPropertyBag& InAnimatorResult, FInstancedPropertyBag& OutEvaluatedValues)
 {
-	TimeOffset = InTimeOffset;
+	const TValueOrError<float, EPropertyBagResult> AlphaResult = InAnimatorResult.GetValueFloat(UPropertyAnimatorCoreBase::AlphaParameterName);
+	const TValueOrError<float, EPropertyBagResult> MagnitudeResult = InAnimatorResult.GetValueFloat(UPropertyAnimatorCoreBase::MagnitudeParameterName);
+
+	if (AlphaResult.HasValue())
+	{
+		const FName DisplayName(InProperty.GetPathHash());
+		OutEvaluatedValues.AddProperty(DisplayName, EPropertyBagPropertyType::Double);
+		OutEvaluatedValues.SetValueFloat(DisplayName, MagnitudeResult.GetValue() * FMath::Lerp(AmplitudeMin, AmplitudeMax, AlphaResult.GetValue()));
+		return true;
+	}
+
+	return false;
 }
 
 void UPropertyAnimatorFloatContext::OnAnimatedPropertyLinked()
 {
 	Super::OnAnimatedPropertyLinked();
 
+	AmplitudeClampMin.Reset();
+	AmplitudeClampMax.Reset();
+
 #if WITH_EDITOR
 	const FPropertyAnimatorCoreData& Property = GetAnimatedProperty();
-
 	const FProperty* LeafProperty = Property.GetLeafProperty();
 
 	checkf(LeafProperty, TEXT("Animated leaf property must be valid"))
@@ -44,6 +67,7 @@ void UPropertyAnimatorFloatContext::OnAnimatedPropertyLinked()
 		if (LeafProperty->HasMetaData(TEXT("ClampMin")))
 		{
 			AmplitudeMin = LeafProperty->GetFloatMetaData(FName("ClampMin"));
+			AmplitudeClampMin = AmplitudeMin;
 		}
 		else if (LeafProperty->HasMetaData(TEXT("UIMin")))
 		{
@@ -53,6 +77,7 @@ void UPropertyAnimatorFloatContext::OnAnimatedPropertyLinked()
 		if (LeafProperty->HasMetaData(TEXT("ClampMax")))
 		{
 			AmplitudeMax =  LeafProperty->GetFloatMetaData(FName("ClampMax"));
+			AmplitudeClampMax = AmplitudeMax;
 		}
 		else if (LeafProperty->HasMetaData(TEXT("UIMax")))
 		{
@@ -61,3 +86,18 @@ void UPropertyAnimatorFloatContext::OnAnimatedPropertyLinked()
 	}
 #endif
 }
+
+double UPropertyAnimatorFloatContext::GetClampedAmplitude(double InAmplitude)
+{
+	if (AmplitudeClampMin.IsSet())
+	{
+		InAmplitude = FMath::Max(InAmplitude, AmplitudeClampMin.GetValue());
+	}
+
+	if (AmplitudeClampMax.IsSet())
+	{
+		InAmplitude = FMath::Min(InAmplitude, AmplitudeClampMax.GetValue());
+	}
+
+	return InAmplitude;
+}
diff --git a/Engine/Plugins/Experimental/PropertyAnimator/Source/PropertyAnimator/Private/Properties/PropertyAnimatorRotatorContext.cpp b/Engine/Plugins/Experimental/PropertyAnimator/Source/PropertyAnimator/Private/Properties/PropertyAnimatorRotatorContext.cpp
new file mode 100644
--- /dev/null
+++ b/Engine/Plugins/Experimental/PropertyAnimator/Source/PropertyAnimator/Private/Properties/PropertyAnimatorRotatorContext.cpp
@@ -0,0 +1,111 @@
+// Copyright Epic Games, Inc. All Rights Reserved.
+
+#include "Properties/PropertyAnimatorRotatorContext.h"
+
+#include "Animators/PropertyAnimatorCoreBase.h"
+
+void UPropertyAnimatorRotatorContext::SetAmplitudeMin(const FRotator& InAmplitude)
+{
+	AmplitudeMin = GetClampedAmplitude(InAmplitude);
+}
+
+void UPropertyAnimatorRotatorContext::SetAmplitudeMax(const FRotator& InAmplitude)
+{
+	AmplitudeMax = GetClampedAmplitude(InAmplitude);
+}
+
+#if WITH_EDITOR
+void UPropertyAnimatorRotatorContext::PostEditChangeProperty(FPropertyChangedEvent& InPropertyChangedEvent)
+{
+	Super::PostEditChangeProperty(InPropertyChangedEvent);
+
+	const FName MemberPropertyName = InPropertyChangedEvent.GetMemberPropertyName();
+
+	if (MemberPropertyName == GET_MEMBER_NAME_CHECKED(UPropertyAnimatorRotatorContext, AmplitudeMin))
+	{
+		SetAmplitudeMin(AmplitudeMin);
+	}
+	else if (MemberPropertyName == GET_MEMBER_NAME_CHECKED(UPropertyAnimatorRotatorContext, AmplitudeMax))
+	{
+		SetAmplitudeMax(AmplitudeMax);
+	}
+}
+#endif
+
+bool UPropertyAnimatorRotatorContext::EvaluateProperty(const FPropertyAnimatorCoreData& InProperty, const FInstancedPropertyBag& InAnimatorResult, FInstancedPropertyBag& OutEvaluatedValues)
+{
+	const TValueOrError<float, EPropertyBagResult> AlphaResult = InAnimatorResult.GetValueFloat(UPropertyAnimatorCoreBase::AlphaParameterName);
+	const TValueOrError<float, EPropertyBagResult> MagnitudeResult = InAnimatorResult.GetValueFloat(UPropertyAnimatorCoreBase::MagnitudeParameterName);
+
+	if (AlphaResult.HasValue())
+	{
+		const FName DisplayName(InProperty.GetPathHash());
+		OutEvaluatedValues.AddProperty(DisplayName, EPropertyBagPropertyType::Struct, TBaseStructure<FRotator>::Get());
+		OutEvaluatedValues.SetValueStruct(DisplayName, MagnitudeResult.GetValue() * FMath::LerpRange(AmplitudeMin, AmplitudeMax, AlphaResult.GetValue()));
+		return true;
+	}
+
+	return false;
+}
+
+void UPropertyAnimatorRotatorContext::OnAnimatedPropertyLinked()
+{
+	Super::OnAnimatedPropertyLinked();
+
+	AmplitudeClampMin.Reset();
+	AmplitudeClampMax.Reset();
+
+#if WITH_EDITOR
+	const FPropertyAnimatorCoreData& Property = GetAnimatedProperty();
+	const FProperty* LeafProperty = Property.GetLeafProperty();
+
+	checkf(LeafProperty, TEXT("Animated leaf property must be valid"))
+
+	// Assign Min and Max value based on editor meta data available
+	if (LeafProperty->IsA<FStructProperty>())
+	{
+		if (LeafProperty->HasMetaData(TEXT("ClampMin")))
+		{
+			AmplitudeMin = FRotator(LeafProperty->GetFloatMetaData(FName("ClampMin")));
+			AmplitudeClampMin = AmplitudeMin;
+		}
+		else if (LeafProperty->HasMetaData(TEXT("UIMin")))
+		{
+			AmplitudeMin = FRotator(LeafProperty->GetFloatMetaData(FName("UIMin")));
+		}
+
+		if (LeafProperty->HasMetaData(TEXT("ClampMax")))
+		{
+			AmplitudeMax = FRotator(LeafProperty->GetFloatMetaData(FName("ClampMax")));
+			AmplitudeClampMax = AmplitudeMax;
+		}
+		else if (LeafProperty->HasMetaData(TEXT("UIMax")))
+		{
+			AmplitudeMax = FRotator(LeafProperty->GetFloatMetaData(FName("UIMax")));
+		}
+	}
+#endif
+}
+
+FRotator UPropertyAnimatorRotatorContext::GetClampedAmplitude(FRotator InAmplitude)
+{
+	if (AmplitudeClampMin.IsSet())
+	{
+		const FRotator& MinAmplitude = AmplitudeClampMin.GetValue();
+
+		InAmplitude.Roll = FMath::Max(InAmplitude.Roll, MinAmplitude.Roll);
+		InAmplitude.Pitch = FMath::Max(InAmplitude.Pitch, MinAmplitude.Pitch);
+		InAmplitude.Yaw = FMath::Max(InAmplitude.Yaw, MinAmplitude.Yaw);
+	}
+
+	if (AmplitudeClampMax.IsSet())
+	{
+		const FRotator& MaxAmplitude = AmplitudeClampMax.GetValue();
+
+		InAmplitude.Roll = FMath::Min(InAmplitude.Roll, MaxAmplitude.Roll);
+		InAmplitude.Pitch = FMath::Min(InAmplitude.Pitch, MaxAmplitude.Pitch);
+		InAmplitude.Yaw = FMath::Min(InAmplitude.Yaw, MaxAmplitude.Yaw);
+	}
+
+	return InAmplitude;
+}
diff --git a/Engine/Plugins/Experimental/PropertyAnimator/Source/PropertyAnimator/Private/Properties/PropertyAnimatorVectorContext.cpp b/Engine/Plugins/Experimental/PropertyAnimator/Source/PropertyAnimator/Private/Properties/PropertyAnimatorVectorContext.cpp
new file mode 100644
--- /dev/null
+++ b/Engine/Plugins/Experimental/PropertyAnimator/Source/PropertyAnimator/Private/Properties/PropertyAnimatorVectorContext.cpp
@@ -0,0 +1,105 @@
+// Copyright Epic Games, Inc. All Rights Reserved.
+
+#include "Properties/PropertyAnimatorVectorContext.h"
+
+#include "Animators/PropertyAnimatorCoreBase.h"
+
+void UPropertyAnimatorVectorContext::SetAmplitudeMin(const FVector& InAmplitude)
+{
+	AmplitudeMin = GetClampedAmplitude(InAmplitude);
+}
+
+void UPropertyAnimatorVectorContext::SetAmplitudeMax(const FVector& InAmplitude)
+{
+	AmplitudeMax = GetClampedAmplitude(InAmplitude);
+}
+
+#if WITH_EDITOR
+void UPropertyAnimatorVectorContext::PostEditChangeProperty(FPropertyChangedEvent& InPropertyChangedEvent)
+{
+	Super::PostEditChangeProperty(InPropertyChangedEvent);
+
+	const FName MemberPropertyName = InPropertyChangedEvent.GetMemberPropertyName();
+
+	if (MemberPropertyName == GET_MEMBER_NAME_CHECKED(UPropertyAnimatorVectorContext, AmplitudeMin))
+	{
+		SetAmplitudeMin(AmplitudeMin);
+	}
+	else if (MemberPropertyName == GET_MEMBER_NAME_CHECKED(UPropertyAnimatorVectorContext, AmplitudeMax))
+	{
+		SetAmplitudeMax(AmplitudeMax);
+	}
+}
+#endif
+
+bool UPropertyAnimatorVectorContext::EvaluateProperty(const FPropertyAnimatorCoreData& InProperty, const FInstancedPropertyBag& InAnimatorResult, FInstancedPropertyBag& OutEvaluatedValues)
+{
+	const TValueOrError<float, EPropertyBagResult> AlphaResult = InAnimatorResult.GetValueFloat(UPropertyAnimatorCoreBase::AlphaParameterName);
+	const TValueOrError<float, EPropertyBagResult> MagnitudeResult = InAnimatorResult.GetValueFloat(UPropertyAnimatorCoreBase::MagnitudeParameterName);
+
+	if (AlphaResult.HasValue())
+	{
+		const FName DisplayName(InProperty.GetPathHash());
+		OutEvaluatedValues.AddProperty(DisplayName, EPropertyBagPropertyType::Struct, TBaseStructure<FVector>::Get());
+		OutEvaluatedValues.SetValueStruct(DisplayName, MagnitudeResult.GetValue() * FMath::Lerp(AmplitudeMin, AmplitudeMax, AlphaResult.GetValue()));
+		return true;
+	}
+
+	return false;
+}
+
+void UPropertyAnimatorVectorContext::OnAnimatedPropertyLinked()
+{
+	Super::OnAnimatedPropertyLinked();
+
+	AmplitudeClampMin.Reset();
+	AmplitudeClampMax.Reset();
+
+#if WITH_EDITOR
+	const FPropertyAnimatorCoreData& Property = GetAnimatedProperty();
+	const FProperty* LeafProperty = Property.GetLeafProperty();
+
+	checkf(LeafProperty, TEXT("Animated leaf property must be valid"))
+
+	// Assign Min and Max value based on editor meta data available
+	if (LeafProperty->IsA<FStructProperty>())
+	{
+		if (LeafProperty->HasMetaData(TEXT("ClampMin")))
+		{
+			AmplitudeMin = FVector(LeafProperty->GetFloatMetaData(FName("ClampMin")));
+			AmplitudeClampMin = AmplitudeMin;
+		}
+		else if (LeafProperty->HasMetaData(TEXT("UIMin")))
+		{
+			AmplitudeMin = FVector(LeafProperty->GetFloatMetaData(FName("UIMin")));
+		}
+
+		if (LeafProperty->HasMetaData(TEXT("ClampMax")))
+		{
+			AmplitudeMax = FVector(LeafProperty->GetFloatMetaData(FName("ClampMax")));
+			AmplitudeClampMax = AmplitudeMax;
+		}
+		else if (LeafProperty->HasMetaData(TEXT("UIMax")))
+		{
+			AmplitudeMax = FVector(LeafProperty->GetFloatMetaData(FName("UIMax")));
+		}
+	}
+#endif
+}
+
+FVector UPropertyAnimatorVectorContext::GetClampedAmplitude(FVector InAmplitude)
+{
+	if (AmplitudeClampMin.IsSet())
+	{
+		const FVector& MinAmplitude = AmplitudeClampMin.GetValue();
+		InAmplitude = InAmplitude.ComponentMax(MinAmplitude);
+	}
+
+	if (AmplitudeClampMax.IsSet())
+	{
+		const FVector& MaxAmplitude = AmplitudeClampMax.GetValue();
+		InAmplitude = InAmplitude.ComponentMin(MaxAmplitude);
+	}
+
+	return InAmplitude;
+}
diff --git a/Engine/Plugins/Experimental/PropertyAnimator/Source/PropertyAnimator/Public/Animators/PropertyAnimatorBounce.h b/Engine/Plugins/Experimental/PropertyAnimator/Source/PropertyAnimator/Public/Animators/PropertyAnimatorBounce.h
--- a/Engine/Plugins/Experimental/PropertyAnimator/Source/PropertyAnimator/Public/Animators/PropertyAnimatorBounce.h
+++ b/Engine/Plugins/Experimental/PropertyAnimator/Source/PropertyAnimator/Public/Animators/PropertyAnimatorBounce.h
@@ -28,7 +28,7 @@ protected:
 	virtual void OnInvertEffect() {}
 
 	//~ Begin UPropertyAnimatorFloatBase
-	virtual float Evaluate(double InTimeElapsed, const FPropertyAnimatorCoreData& InPropertyData, UPropertyAnimatorFloatContext* InOptions) const override;
+	virtual bool EvaluateProperty(const FPropertyAnimatorCoreData& InPropertyData, UPropertyAnimatorCoreContext* InContext, FInstancedPropertyBag& InParameters, FInstancedPropertyBag& OutEvaluationResult) const override;
 	//~ End UPropertyAnimatorFloatBase
 
 	/** Invert the effect result */
diff --git a/Engine/Plugins/Experimental/PropertyAnimator/Source/PropertyAnimator/Public/Animators/PropertyAnimatorClock.h b/Engine/Plugins/Experimental/PropertyAnimator/Source/PropertyAnimator/Public/Animators/PropertyAnimatorClock.h
--- a/Engine/Plugins/Experimental/PropertyAnimator/Source/PropertyAnimator/Public/Animators/PropertyAnimatorClock.h
+++ b/Engine/Plugins/Experimental/PropertyAnimator/Source/PropertyAnimator/Public/Animators/PropertyAnimatorClock.h
@@ -18,7 +18,7 @@ enum class EPropertyAnimatorClockMode : uint8
 	Stopwatch
 };
 
-/** Animate supported string properties with a clock feature */
+/** Animate supported string properties to display time */
 UCLASS(MinimalAPI, AutoExpandCategories=("Animator"))
 class UPropertyAnimatorClock : public UPropertyAnimatorCoreBase
 {
@@ -29,47 +29,19 @@ public:
 
 	UPropertyAnimatorClock();
 
-	PROPERTYANIMATOR_API void SetMode(EPropertyAnimatorClockMode InMode);
-	EPropertyAnimatorClockMode GetMode() const
-	{
-		return Mode;
-	}
-
 	PROPERTYANIMATOR_API void SetDisplayFormat(const FString& InDisplayFormat);
 	const FString& GetDisplayFormat() const
 	{
 		return DisplayFormat;
 	}
 
-	PROPERTYANIMATOR_API void SetCountdownDuration(const FString& InDuration);
-	const FString& GetCountdownDuration() const
-	{
-		return CountdownDuration;
-	}
-
 protected:
-	static FTimespan ParseTime(const FString& InFormat);
-
-	//~ Begin UObject
-	virtual void PostLoad() override;
-#if WITH_EDITOR
-	virtual void PostEditChangeProperty(FPropertyChangedEvent& InPropertyChangedEvent) override;
-#endif
-	//~ End UObject
-
 	//~ Begin UPropertyAnimatorCoreBase
-	virtual bool IsPropertyDirectlySupported(const FPropertyAnimatorCoreData& InPropertyData) const override;
-	virtual bool IsPropertyIndirectlySupported(const FPropertyAnimatorCoreData& InPropertyData) const override;
-	virtual void EvaluateProperties(const FPropertyAnimatorCoreEvaluationParameters& InParameters) override;
-	virtual void OnPropertyLinked(UPropertyAnimatorCoreContext* InLinkedProperty) override;
+	virtual EPropertyAnimatorPropertySupport IsPropertySupported(const FPropertyAnimatorCoreData& InPropertyData) const override;
+	virtual void EvaluateProperties(FInstancedPropertyBag& InParameters) override;
+	virtual void OnPropertyLinked(UPropertyAnimatorCoreContext* InLinkedProperty, EPropertyAnimatorPropertySupport InSupport) override;
 	//~ End UPropertyAnimatorCoreBase
 
-	void OnModeChanged();
-
-	/** Mode chosen for this clock animator */
-	UPROPERTY(EditInstanceOnly, BlueprintReadWrite, Setter, Getter, Category="Animator")
-	EPropertyAnimatorClockMode Mode = EPropertyAnimatorClockMode::LocalTime;
-
 	/**
 	 * Display date time format : 
 	 * %a - Weekday, eg) Sun
@@ -93,23 +65,4 @@ protected:
 	 */
 	UPROPERTY(EditInstanceOnly, BlueprintReadWrite, Setter, Getter, Category="Animator")
 	FString DisplayFormat = TEXT("%H:%M:%S");
-
-	/**
-	 * Countdown duration format : 
-	 * 120 = 2 minutes
-	 * 02:00 = 2 minutes
-	 * 00:02:00 = 2 minutes
-	 * 2m = 2 minutes
-	 * 1h = 1 hour
-	 * 120s = 2 minutes
-	 */
-	UPROPERTY(EditInstanceOnly, BlueprintReadWrite, Setter, Getter, Category="Animator", meta=(EditCondition="Mode == EPropertyAnimatorClockMode::Countdown", EditConditionHides))
-	FString CountdownDuration = TEXT("5m");
-
-private:
-	UPROPERTY(Transient)
-	FTimespan ActiveTimeSpan;
-
-	UPROPERTY(Transient)
-	FTimespan ElapsedTimeSpan;
 };
\ No newline at end of file
diff --git a/Engine/Plugins/Experimental/PropertyAnimator/Source/PropertyAnimator/Public/Animators/PropertyAnimatorFloatBase.h b/Engine/Plugins/Experimental/PropertyAnimator/Source/PropertyAnimator/Public/Animators/PropertyAnimatorFloatBase.h
--- a/Engine/Plugins/Experimental/PropertyAnimator/Source/PropertyAnimator/Public/Animators/PropertyAnimatorFloatBase.h
+++ b/Engine/Plugins/Experimental/PropertyAnimator/Source/PropertyAnimator/Public/Animators/PropertyAnimatorFloatBase.h
@@ -7,6 +7,17 @@
 
 class UPropertyAnimatorFloatContext;
 
+UENUM(BlueprintType)
+enum class EPropertyAnimatorCycleMode : uint8
+{
+	/** Cycle only once then stop */
+	DoOnce,
+	/** Cycle and repeat once we reached the end */
+	Loop,
+	/** Cycle and reverse repeat */
+	PingPong
+};
+
 /**
  * Animate supported float properties with various options
  */
@@ -18,22 +29,28 @@ class UPropertyAnimatorFloatBase : public UPropertyAnimatorCoreBase
 	friend class FPropertyAnimatorCoreEditorDetailCustomization;
 
 public:
-	PROPERTYANIMATOR_API void SetGlobalMagnitude(float InMagnitude);
-	float GetGlobalMagnitude() const
+	PROPERTYANIMATOR_API void SetMagnitude(float InMagnitude);
+	float GetMagnitude() const
 	{
-		return GlobalMagnitude;
+		return Magnitude;
 	}
 
-	PROPERTYANIMATOR_API void SetGlobalFrequency(float InFrequency);
-	float GetGlobalFrequency() const
+	PROPERTYANIMATOR_API void SetCycleDuration(float InCycleDuration);
+	float GetCycleDuration() const
 	{
-		return GlobalFrequency;
+		return CycleDuration;
 	}
 
-	PROPERTYANIMATOR_API void SetAccumulatedTimeOffset(double InOffset);
-	double GetAccumulatedTimeOffset() const
+	PROPERTYANIMATOR_API void SetCycleMode(EPropertyAnimatorCycleMode InMode);
+	EPropertyAnimatorCycleMode GetCycleMode() const
 	{
-		return AccumulatedTimeOffset;
+		return CycleMode;
+	}
+
+	PROPERTYANIMATOR_API void SetTimeOffset(double InOffset);
+	double GetTimeOffset() const
+	{
+		return TimeOffset;
 	}
 
 	PROPERTYANIMATOR_API void SetRandomTimeOffset(bool bInOffset);
@@ -56,43 +73,51 @@ protected:
 	//~ End UObject
 
 	virtual void OnMagnitudeChanged() {}
-	virtual void OnGlobalFrequencyChanged() {}
-	virtual void OnAccumulatedTimeOffsetChanged() {}
+	virtual void OnCycleDurationChanged() {}
+	virtual void OnCycleModeChanged() {}
+	virtual void OnTimeOffsetChanged() {}
 	virtual void OnSeedChanged() {}
 
 	//~ Begin UPropertyAnimatorCoreBase
-	virtual TSubclassOf<UPropertyAnimatorCoreContext> GetPropertyContextClass(const FPropertyAnimatorCoreData& InUnlinkedProperty) override;
-	virtual bool IsPropertyDirectlySupported(const FPropertyAnimatorCoreData& InPropertyData) const override;
-	virtual bool IsPropertyIndirectlySupported(const FPropertyAnimatorCoreData& InPropertyData) const override;
-	virtual void EvaluateProperties(const FPropertyAnimatorCoreEvaluationParameters& InParameters) override;
-	virtual void OnPropertyLinked(UPropertyAnimatorCoreContext* InLinkedProperty) override;
+	virtual TSubclassOf<UPropertyAnimatorCoreContext> GetPropertyContextClass(const FPropertyAnimatorCoreData& InProperty) override;
+	virtual EPropertyAnimatorPropertySupport IsPropertySupported(const FPropertyAnimatorCoreData& InPropertyData) const override;
+	virtual void EvaluateProperties(FInstancedPropertyBag& InParameters) override;
+	virtual void OnPropertyLinked(UPropertyAnimatorCoreContext* InLinkedProperty, EPropertyAnimatorPropertySupport InSupport) override;
 	//~ End UPropertyAnimatorCoreBase
 
 	/** Evaluate and return float value for a property */
-	virtual float Evaluate(double InTimeElapsed, const FPropertyAnimatorCoreData& InPropertyData, UPropertyAnimatorFloatContext* InOptions) const
+	virtual bool EvaluateProperty(const FPropertyAnimatorCoreData& InPropertyData, UPropertyAnimatorCoreContext* InContext, FInstancedPropertyBag& InParameters, FInstancedPropertyBag& OutEvaluationResult) const
 	{
-		return 0.f;
+		return false;
 	}
 
-	/** Global magnitude for the effect */
-	UPROPERTY(EditInstanceOnly, BlueprintReadWrite, Setter, Getter, Category="Animator", meta=(ClampMin="0"))
-	float GlobalMagnitude = 1.f;
+	/** Magnitude for the effect on all properties */
+	UPROPERTY(EditInstanceOnly, Setter, Getter, Category="Animator", meta=(ClampMin="0"))
+	float Magnitude = 1.f;
+
+	/** Duration of one cycle for the effect = period of the effect */
+	UPROPERTY(EditInstanceOnly, Setter, Getter, Category="Animator", meta=(ClampMin="0", Units=Seconds))
+	float CycleDuration = 1.f;
+
+	/** Cycle mode for the effect */
+	UPROPERTY(EditInstanceOnly, Setter, Getter, Category="Animator")
+	EPropertyAnimatorCycleMode CycleMode = EPropertyAnimatorCycleMode::Loop;
 
-	/** Global frequency multiplier */
-	UPROPERTY(EditInstanceOnly, BlueprintReadWrite, Setter, Getter, Category="Animator", meta=(ClampMin="0"))
-	float GlobalFrequency = 1.f;
+	/** Time gap between each cycle */
+	UPROPERTY(EditInstanceOnly, Category="Animator", meta=(ClampMin="0", Units=Seconds))
+	float CycleGapDuration = 0.f;
 
 	/** Use random time offset to add variation in animation */
-	UPROPERTY(EditAnywhere, BlueprintReadWrite, Setter="SetRandomTimeOffset", Getter="GetRandomTimeOffset", Category="Animator")
+	UPROPERTY(EditInstanceOnly, Setter="SetRandomTimeOffset", Getter="GetRandomTimeOffset", Category="Animator", meta=(InlineEditConditionToggle))
 	bool bRandomTimeOffset = false;
 
 	/** Seed to generate per property time offset */
-	UPROPERTY(EditAnywhere, BlueprintReadWrite, Setter, Getter, Category="Animator", meta=(EditCondition="bRandomTimeOffset", EditConditionHides))
+	UPROPERTY(EditInstanceOnly, Setter, Getter, Category="Animator", meta=(EditCondition="bRandomTimeOffset"))
 	int32 Seed = 0;
 
-	/** This time offset will be accumulated for each property for every round */
-	UPROPERTY(EditInstanceOnly, BlueprintReadWrite, Setter, Getter, Category="Animator")
-	double AccumulatedTimeOffset = 0.f;
+	/** Time offset accumulated for each property for every round */
+	UPROPERTY(EditInstanceOnly, Setter, Getter, Category="Animator", meta=(Units=Seconds))
+	double TimeOffset = 0.f;
 
 private:
 	/** Random stream for time offset */
diff --git a/Engine/Plugins/Experimental/PropertyAnimator/Source/PropertyAnimator/Public/Animators/PropertyAnimatorOscillate.h b/Engine/Plugins/Experimental/PropertyAnimator/Source/PropertyAnimator/Public/Animators/PropertyAnimatorOscillate.h
--- a/Engine/Plugins/Experimental/PropertyAnimator/Source/PropertyAnimator/Public/Animators/PropertyAnimatorOscillate.h
+++ b/Engine/Plugins/Experimental/PropertyAnimator/Source/PropertyAnimator/Public/Animators/PropertyAnimatorOscillate.h
@@ -37,7 +37,7 @@ public:
 
 protected:
 	//~ Begin UPropertyAnimatorFloatBase
-	virtual float Evaluate(double InTimeElapsed, const FPropertyAnimatorCoreData& InPropertyData, UPropertyAnimatorFloatContext* InOptions) const override;
+	virtual bool EvaluateProperty(const FPropertyAnimatorCoreData& InPropertyData, UPropertyAnimatorCoreContext* InContext, FInstancedPropertyBag& InParameters, FInstancedPropertyBag& OutEvaluationResult) const override;
 	//~ End UPropertyAnimatorFloatBase
 
 	/** The oscillate function to feed current time elapsed */
diff --git a/Engine/Plugins/Experimental/PropertyAnimator/Source/PropertyAnimator/Public/Animators/PropertyAnimatorPulse.h b/Engine/Plugins/Experimental/PropertyAnimator/Source/PropertyAnimator/Public/Animators/PropertyAnimatorPulse.h
--- a/Engine/Plugins/Experimental/PropertyAnimator/Source/PropertyAnimator/Public/Animators/PropertyAnimatorPulse.h
+++ b/Engine/Plugins/Experimental/PropertyAnimator/Source/PropertyAnimator/Public/Animators/PropertyAnimatorPulse.h
@@ -35,7 +35,7 @@ public:
 
 protected:
 	//~ Begin UPropertyAnimatorFloatBase
-	virtual float Evaluate(double InTimeElapsed, const FPropertyAnimatorCoreData& InPropertyData, UPropertyAnimatorFloatContext* InOptions) const override;
+	virtual bool EvaluateProperty(const FPropertyAnimatorCoreData& InPropertyData, UPropertyAnimatorCoreContext* InContext, FInstancedPropertyBag& InParameters, FInstancedPropertyBag& OutEvaluationResult) const override;
 	//~ End UPropertyAnimatorFloatBase
 
 	/** The easing function to use to modify the base effect */
diff --git a/Engine/Plugins/Experimental/PropertyAnimator/Source/PropertyAnimator/Public/Animators/PropertyAnimatorSoundWave.h b/Engine/Plugins/Experimental/PropertyAnimator/Source/PropertyAnimator/Public/Animators/PropertyAnimatorSoundWave.h
--- a/Engine/Plugins/Experimental/PropertyAnimator/Source/PropertyAnimator/Public/Animators/PropertyAnimatorSoundWave.h
+++ b/Engine/Plugins/Experimental/PropertyAnimator/Source/PropertyAnimator/Public/Animators/PropertyAnimatorSoundWave.h
@@ -43,7 +43,7 @@ protected:
 	void OnSampledSoundWaveChanged();
 
 	//~ Begin UPropertyAnimatorFloatBase
-	virtual float Evaluate(double InTimeElapsed, const FPropertyAnimatorCoreData& InPropertyData, UPropertyAnimatorFloatContext* InOptions) const override;
+	virtual bool EvaluateProperty(const FPropertyAnimatorCoreData& InPropertyData, UPropertyAnimatorCoreContext* InContext, FInstancedPropertyBag& InParameters, FInstancedPropertyBag& OutEvaluationResult) const override;
 	//~ End UPropertyAnimatorFloatBase
 
 	/**
diff --git a/Engine/Plugins/Experimental/PropertyAnimator/Source/PropertyAnimator/Public/Animators/PropertyAnimatorTime.h b/Engine/Plugins/Experimental/PropertyAnimator/Source/PropertyAnimator/Public/Animators/PropertyAnimatorTime.h
--- a/Engine/Plugins/Experimental/PropertyAnimator/Source/PropertyAnimator/Public/Animators/PropertyAnimatorTime.h
+++ b/Engine/Plugins/Experimental/PropertyAnimator/Source/PropertyAnimator/Public/Animators/PropertyAnimatorTime.h
@@ -20,6 +20,6 @@ public:
 
 protected:
 	//~ Begin UPropertyAnimatorFloatBase
-	virtual float Evaluate(double InTimeElapsed, const FPropertyAnimatorCoreData& InPropertyData, UPropertyAnimatorFloatContext* InOptions) const override;
+	virtual bool EvaluateProperty(const FPropertyAnimatorCoreData& InPropertyData, UPropertyAnimatorCoreContext* InContext, FInstancedPropertyBag& InParameters, FInstancedPropertyBag& OutEvaluationResult) const override;
 	//~ End UPropertyAnimatorFloatBase
 };
\ No newline at end of file
diff --git a/Engine/Plugins/Experimental/PropertyAnimator/Source/PropertyAnimator/Public/Animators/PropertyAnimatorWiggle.h b/Engine/Plugins/Experimental/PropertyAnimator/Source/PropertyAnimator/Public/Animators/PropertyAnimatorWiggle.h
--- a/Engine/Plugins/Experimental/PropertyAnimator/Source/PropertyAnimator/Public/Animators/PropertyAnimatorWiggle.h
+++ b/Engine/Plugins/Experimental/PropertyAnimator/Source/PropertyAnimator/Public/Animators/PropertyAnimatorWiggle.h
@@ -20,6 +20,6 @@ public:
 
 protected:
 	//~ Begin UPropertyAnimatorFloatBase
-	virtual float Evaluate(double InTimeElapsed, const FPropertyAnimatorCoreData& InPropertyData, UPropertyAnimatorFloatContext* InOptions) const override;
+	virtual bool EvaluateProperty(const FPropertyAnimatorCoreData& InPropertyData, UPropertyAnimatorCoreContext* InContext, FInstancedPropertyBag& InParameters, FInstancedPropertyBag& OutEvaluationResult) const override;
 	//~ End UPropertyAnimatorFloatBase
 };
\ No newline at end of file
diff --git a/Engine/Plugins/Experimental/PropertyAnimator/Source/PropertyAnimator/Public/Properties/PropertyAnimatorFloatContext.h b/Engine/Plugins/Experimental/PropertyAnimator/Source/PropertyAnimator/Public/Properties/PropertyAnimatorFloatContext.h
--- a/Engine/Plugins/Experimental/PropertyAnimator/Source/PropertyAnimator/Public/Properties/PropertyAnimatorFloatContext.h
+++ b/Engine/Plugins/Experimental/PropertyAnimator/Source/PropertyAnimator/Public/Properties/PropertyAnimatorFloatContext.h
@@ -5,19 +5,13 @@
 #include "Properties/PropertyAnimatorCoreContext.h"
 #include "PropertyAnimatorFloatContext.generated.h"
 
-/** Property context used by float driving animator */
+/** Property context used by animator for float/double properties */
 UCLASS(MinimalAPI, BlueprintType)
 class UPropertyAnimatorFloatContext : public UPropertyAnimatorCoreContext
 {
 	GENERATED_BODY()
 
 public:
-	PROPERTYANIMATOR_API void SetMagnitude(float InMagnitude);
-	float GetMagnitude() const
-	{
-		return Magnitude;
-	}
-
 	PROPERTYANIMATOR_API void SetAmplitudeMin(double InAmplitude);
 	double GetAmplitudeMin() const
 	{
@@ -30,40 +24,33 @@ public:
 		return AmplitudeMax;
 	}
 
-	PROPERTYANIMATOR_API void SetFrequency(float InFrequency);
-	float GetFrequency() const
-	{
-		return Frequency;
-	}
-
-	PROPERTYANIMATOR_API void SetTimeOffset(double InTimeOffset);
-	double GetTimeOffset() const
-	{
-		return TimeOffset;
-	}
-
 protected:
+	//~ Begin UObject
+#if WITH_EDITOR
+	virtual void PostEditChangeProperty(FPropertyChangedEvent& InPropertyChangedEvent) override;
+#endif
+	//~ End UObject
+
 	//~ Begin UPropertyAnimatorCoreContext
+	virtual bool EvaluateProperty(const FPropertyAnimatorCoreData& InProperty, const FInstancedPropertyBag& InAnimatorResult, FInstancedPropertyBag& OutEvaluatedValues) override;
 	virtual void OnAnimatedPropertyLinked() override;
 	//~ End UPropertyAnimatorCoreContext
 
-	/** Magnitude of the effect */
-	UPROPERTY(EditAnywhere, BlueprintReadWrite, Interp, Category="Animator", meta=(ClampMin="0.0", ClampMax="1.0"))
-	float Magnitude = 1.f;
+	double GetClampedAmplitude(double InAmplitude);
 
 	/** The minimum value should be remapped to that values */
-	UPROPERTY(EditAnywhere, BlueprintReadWrite, Interp, Category="Animator")
+	UPROPERTY(EditInstanceOnly, Setter, Getter, Category="Animator")
 	double AmplitudeMin = -1.f;
 
+	/** Some properties are clamped and cannot go past a specific min value */
+	UPROPERTY()
+	TOptional<double> AmplitudeClampMin;
+
 	/** The maximum value should be remapped to that values */
-	UPROPERTY(EditAnywhere, BlueprintReadWrite, Interp, Category="Animator")
+	UPROPERTY(EditInstanceOnly, Setter, Getter, Category="Animator")
 	double AmplitudeMax = 1.f;
 
-	/** Number of repetition every seconds for the effect */
-	UPROPERTY(EditAnywhere, BlueprintReadWrite, Interp, Category="Animator", meta=(ClampMin="0"))
-	float Frequency = 1.f;
-
-	/** Time offset variation for evaluation */
-	UPROPERTY(EditAnywhere, BlueprintReadWrite, Interp, Category="Animator")
-	double TimeOffset = 0.f;
+	/** Some properties are clamped and cannot go past a specific min value */
+	UPROPERTY()
+	TOptional<double> AmplitudeClampMax;
 };
diff --git a/Engine/Plugins/Experimental/PropertyAnimator/Source/PropertyAnimator/Public/Properties/PropertyAnimatorRotatorContext.h b/Engine/Plugins/Experimental/PropertyAnimator/Source/PropertyAnimator/Public/Properties/PropertyAnimatorRotatorContext.h
new file mode 100644
--- /dev/null
+++ b/Engine/Plugins/Experimental/PropertyAnimator/Source/PropertyAnimator/Public/Properties/PropertyAnimatorRotatorContext.h
@@ -0,0 +1,56 @@
+// Copyright Epic Games, Inc. All Rights Reserved.
+
+#pragma once
+
+#include "Properties/PropertyAnimatorCoreContext.h"
+#include "PropertyAnimatorRotatorContext.generated.h"
+
+/** Property context used by animator for rotator properties */
+UCLASS(MinimalAPI, BlueprintType)
+class UPropertyAnimatorRotatorContext : public UPropertyAnimatorCoreContext
+{
+	GENERATED_BODY()
+
+public:
+	PROPERTYANIMATOR_API void SetAmplitudeMin(const FRotator& InAmplitude);
+	const FRotator& GetAmplitudeMin() const
+	{
+		return AmplitudeMin;
+	}
+
+	PROPERTYANIMATOR_API void SetAmplitudeMax(const FRotator& InAmplitude);
+	const FRotator& GetAmplitudeMax() const
+	{
+		return AmplitudeMax;
+	}
+
+protected:
+	//~ Begin UObject
+#if WITH_EDITOR
+	virtual void PostEditChangeProperty(FPropertyChangedEvent& InPropertyChangedEvent) override;
+#endif
+	//~ End UObject
+
+	//~ Begin UPropertyAnimatorCoreContext
+	virtual bool EvaluateProperty(const FPropertyAnimatorCoreData& InProperty, const FInstancedPropertyBag& InAnimatorResult, FInstancedPropertyBag& OutEvaluatedValues) override;
+	virtual void OnAnimatedPropertyLinked() override;
+	//~ End UPropertyAnimatorCoreContext
+
+	FRotator GetClampedAmplitude(FRotator InAmplitude);
+
+	/** The minimum value should be remapped to that values */
+	UPROPERTY(EditInstanceOnly, Setter, Getter, Category="Animator")
+	FRotator AmplitudeMin = FRotator::ZeroRotator;
+
+	/** Some properties are clamped and cannot go past a specific min value */
+	UPROPERTY()
+	TOptional<FRotator> AmplitudeClampMin;
+
+	/** The maximum value should be remapped to that values */
+	UPROPERTY(EditInstanceOnly, Setter, Getter, Category="Animator")
+	FRotator AmplitudeMax = FRotator::ZeroRotator;
+
+	/** Some properties are clamped and cannot go past a specific min value */
+	UPROPERTY()
+	TOptional<FRotator> AmplitudeClampMax;
+};
diff --git a/Engine/Plugins/Experimental/PropertyAnimator/Source/PropertyAnimator/Public/Properties/PropertyAnimatorVectorContext.h b/Engine/Plugins/Experimental/PropertyAnimator/Source/PropertyAnimator/Public/Properties/PropertyAnimatorVectorContext.h
new file mode 100644
--- /dev/null
+++ b/Engine/Plugins/Experimental/PropertyAnimator/Source/PropertyAnimator/Public/Properties/PropertyAnimatorVectorContext.h
@@ -0,0 +1,56 @@
+// Copyright Epic Games, Inc. All Rights Reserved.
+
+#pragma once
+
+#include "Properties/PropertyAnimatorCoreContext.h"
+#include "PropertyAnimatorVectorContext.generated.h"
+
+/** Property context used by animator for vector properties */
+UCLASS(MinimalAPI, BlueprintType)
+class UPropertyAnimatorVectorContext : public UPropertyAnimatorCoreContext
+{
+	GENERATED_BODY()
+
+public:
+	PROPERTYANIMATOR_API void SetAmplitudeMin(const FVector& InAmplitude);
+	const FVector& GetAmplitudeMin() const
+	{
+		return AmplitudeMin;
+	}
+
+	PROPERTYANIMATOR_API void SetAmplitudeMax(const FVector& InAmplitude);
+	const FVector& GetAmplitudeMax() const
+	{
+		return AmplitudeMax;
+	}
+
+protected:
+	//~ Begin UObject
+#if WITH_EDITOR
+	virtual void PostEditChangeProperty(FPropertyChangedEvent& InPropertyChangedEvent) override;
+#endif
+	//~ End UObject
+
+	//~ Begin UPropertyAnimatorCoreContext
+	virtual bool EvaluateProperty(const FPropertyAnimatorCoreData& InProperty, const FInstancedPropertyBag& InAnimatorResult, FInstancedPropertyBag& OutEvaluatedValues) override;
+	virtual void OnAnimatedPropertyLinked() override;
+	//~ End UPropertyAnimatorCoreContext
+
+	FVector GetClampedAmplitude(FVector InAmplitude);
+
+	/** The minimum value should be remapped to that values */
+	UPROPERTY(EditInstanceOnly, Setter, Getter, Category="Animator")
+	FVector AmplitudeMin = FVector::ZeroVector;
+
+	/** Some properties are clamped and cannot go past a specific min value */
+	UPROPERTY()
+	TOptional<FVector> AmplitudeClampMin;
+
+	/** The maximum value should be remapped to that values */
+	UPROPERTY(EditInstanceOnly, Setter, Getter, Category="Animator")
+	FVector AmplitudeMax = FVector::ZeroVector;
+
+	/** Some properties are clamped and cannot go past a specific min value */
+	UPROPERTY()
+	TOptional<FVector> AmplitudeClampMax;
+};
diff --git a/Engine/Plugins/Experimental/PropertyAnimatorCore/Source/PropertyAnimatorCore/Private/Animators/PropertyAnimatorCoreBase.cpp b/Engine/Plugins/Experimental/PropertyAnimatorCore/Source/PropertyAnimatorCore/Private/Animators/PropertyAnimatorCoreBase.cpp
--- a/Engine/Plugins/Experimental/PropertyAnimatorCore/Source/PropertyAnimatorCore/Private/Animators/PropertyAnimatorCoreBase.cpp
+++ b/Engine/Plugins/Experimental/PropertyAnimatorCore/Source/PropertyAnimatorCore/Private/Animators/PropertyAnimatorCoreBase.cpp
@@ -23,8 +23,8 @@ UPropertyAnimatorCoreBase::UPropertyAnimatorCoreBase()
 	if (const UPropertyAnimatorCoreSubsystem* AnimatorSubsystem = UPropertyAnimatorCoreSubsystem::Get())
 	{
 		// Apply default time source
-		const TArray<FName> TimeSources = AnimatorSubsystem->GetTimeSourceNames();
-		SetTimeSourceName(!TimeSources.IsEmpty() ? TimeSources[0] : NAME_None);
+		const TArray<FName> TimeSourceNames = AnimatorSubsystem->GetTimeSourceNames();
+		SetTimeSourceName(!TimeSourceNames.IsEmpty() ? TimeSourceNames[0] : NAME_None);
 	}
 
 #if WITH_EDITOR
@@ -119,11 +119,19 @@ void UPropertyAnimatorCoreBase::BeginDestroy()
 
 void UPropertyAnimatorCoreBase::PostLoad()
 {
-	CleanTimeSources();
+	Super::PostLoad();
 
-	OnTimeSourceNameChanged();
+	PRAGMA_DISABLE_DEPRECATION_WARNINGS
 
-	Super::PostLoad();
+	// Migrate deprecated property
+	if (TimeSources.IsEmpty())
+	{
+		TimeSourcesInstances.GenerateValueArray(TimeSources);
+	}
+
+	PRAGMA_ENABLE_DEPRECATION_WARNINGS
+
+	OnTimeSourceNameChanged();
 
 	CleanLinkedProperties();
 
@@ -213,8 +221,8 @@ void UPropertyAnimatorCoreBase::SetTimeSourceName(FName InTimeSourceName)
 		return;
 	}
 
-	const TArray<FString> TimeSources = GetTimeSourceNames();
-	if (!TimeSources.Contains(InTimeSourceName))
+	const TArray<FName> TimeSourceNames = GetTimeSourceNames();
+	if (!TimeSourceNames.Contains(InTimeSourceName))
 	{
 		return;
 	}
@@ -229,7 +237,7 @@ FName UPropertyAnimatorCoreBase::GetAnimatorOriginalName() const
 	return CDO ? CDO->AnimatorDisplayName : NAME_None;
 }
 
-bool UPropertyAnimatorCoreBase::GetPropertiesSupported(const FPropertyAnimatorCoreData& InPropertyData, TSet<FPropertyAnimatorCoreData>& OutProperties, bool bInRecursiveSearch) const
+bool UPropertyAnimatorCoreBase::GetPropertiesSupported(const FPropertyAnimatorCoreData& InPropertyData, TSet<FPropertyAnimatorCoreData>& OutProperties, uint8 InSearchDepth, EPropertyAnimatorPropertySupport InSupportExpected) const
 {
 	const FProperty* LeafProperty = InPropertyData.GetLeafProperty();
 	UObject* Owner = InPropertyData.GetOwner();
@@ -241,47 +249,46 @@ bool UPropertyAnimatorCoreBase::GetPropertiesSupported(const FPropertyAnimatorCo
 	}
 
 	// We can directly control the member property
-	if (IsPropertySupported(InPropertyData))
+	if (HasPropertySupport(InPropertyData, InSupportExpected))
 	{
 		OutProperties.Add(InPropertyData);
+	}
 
-		if (bInRecursiveSearch)
-		{
-			return true;
-		}
+	if (--InSearchDepth == 0)
+	{
+		return !OutProperties.IsEmpty();
 	}
 
 	// Look for inner properties that can be controlled too
-	TFunction<bool(TArray<FProperty*>&, UObject*, TSet<FPropertyAnimatorCoreData>&)> FindSupportedPropertiesRecursively = [this, &FindSupportedPropertiesRecursively, bInRecursiveSearch, InPropertyData](TArray<FProperty*>& InChainProperties, UObject* InOwner, TSet<FPropertyAnimatorCoreData>& OutSupportedProperties)
+	TFunction<bool(TArray<FProperty*>&, UObject*, TSet<FPropertyAnimatorCoreData>&)> FindSupportedPropertiesRecursively = [this, &FindSupportedPropertiesRecursively, &InSearchDepth, &InPropertyData, InSupportExpected](TArray<FProperty*>& InChainProperties, UObject* InOwner, TSet<FPropertyAnimatorCoreData>& OutSupportedProperties)
 	{
-		FProperty* InLeafProperty = InChainProperties.Last();
-		if (const FStructProperty* StructProp = CastField<FStructProperty>(InLeafProperty))
+		if (InSearchDepth-- > 0)
 		{
-			for (FProperty* Property : TFieldRange<FProperty>(StructProp->Struct))
+			FProperty* InLeafProperty = InChainProperties.Last();
+
+			if (const FStructProperty* StructProp = CastField<FStructProperty>(InLeafProperty))
 			{
-				if (!Property->HasAnyPropertyFlags(EPropertyFlags::CPF_Edit))
+				for (FProperty* Property : TFieldRange<FProperty>(StructProp->Struct))
 				{
-					continue;
-				}
-
-				// Copy over resolver if any on that property
-				FPropertyAnimatorCoreData PropertyControlData(InOwner, InChainProperties, Property, InPropertyData.GetPropertyResolverClass());
+					if (!Property->HasAnyPropertyFlags(EPropertyFlags::CPF_Edit))
+					{
+						continue;
+					}
 
-				// We can directly control this property
-				if (IsPropertySupported(PropertyControlData))
-				{
-					OutSupportedProperties.Add(PropertyControlData);
+					// Copy over resolver if any on that property
+					FPropertyAnimatorCoreData PropertyControlData(InOwner, InChainProperties, Property, InPropertyData.GetPropertyResolverClass());
 
-					if (bInRecursiveSearch)
+					// We can directly control this property
+					if (HasPropertySupport(PropertyControlData, InSupportExpected))
 					{
-						continue;
+						OutSupportedProperties.Add(PropertyControlData);
 					}
-				}
 
-				// Check nested properties inside this property
-				TArray<FProperty*> NestedChainProperties(InChainProperties);
-				NestedChainProperties.Add(Property);
-				FindSupportedPropertiesRecursively(NestedChainProperties, InOwner, OutSupportedProperties);
+					// Check nested properties inside this property
+					TArray<FProperty*> NestedChainProperties(InChainProperties);
+					NestedChainProperties.Add(Property);
+					FindSupportedPropertiesRecursively(NestedChainProperties, InOwner, OutSupportedProperties);
+				}
 			}
 		}
 
@@ -289,20 +296,23 @@ bool UPropertyAnimatorCoreBase::GetPropertiesSupported(const FPropertyAnimatorCo
 	};
 
 	TArray<FProperty*> ChainProperties = InPropertyData.GetChainProperties();
-
 	return FindSupportedPropertiesRecursively(ChainProperties, Owner, OutProperties);
 }
 
-bool UPropertyAnimatorCoreBase::IsPropertySupported(const FPropertyAnimatorCoreData& InPropertyData) const
+EPropertyAnimatorPropertySupport UPropertyAnimatorCoreBase::GetPropertySupport(const FPropertyAnimatorCoreData& InPropertyData) const
 {
 	// Without any handler we can't control the property type
 	if (!InPropertyData.GetPropertyHandler())
 	{
-		return false;
+		return EPropertyAnimatorPropertySupport::None;
 	}
 
-	return IsPropertyDirectlySupported(InPropertyData)
-		|| IsPropertyIndirectlySupported(InPropertyData);
+	return IsPropertySupported(InPropertyData);
+}
+
+bool UPropertyAnimatorCoreBase::HasPropertySupport(const FPropertyAnimatorCoreData& InPropertyData, EPropertyAnimatorPropertySupport InSupportExpected) const
+{
+	return EnumHasAnyFlags(InSupportExpected, GetPropertySupport(InPropertyData));
 }
 
 void UPropertyAnimatorCoreBase::OnAnimatorEnabled()
@@ -357,18 +367,6 @@ void UPropertyAnimatorCoreBase::CleanLinkedProperties()
 	}
 }
 
-void UPropertyAnimatorCoreBase::CleanTimeSources()
-{
-	for (TMap<FName, TObjectPtr<UPropertyAnimatorCoreTimeSourceBase>>::TIterator It(TimeSourcesInstances); It; ++It)
-	{
-		// Remove invalid time sources
-		if (!It->Value.Get())
-		{
-			It.RemoveCurrent();
-		}
-	}
-}
-
 void UPropertyAnimatorCoreBase::OnTimeSourceNameChanged()
 {
 	if (ActiveTimeSource)
@@ -409,7 +407,7 @@ void UPropertyAnimatorCoreBase::ResolvePropertiesOwner(AActor* InNewOwner)
 	}
 }
 
-void UPropertyAnimatorCoreBase::EvaluateAnimator()
+void UPropertyAnimatorCoreBase::EvaluateAnimator(FInstancedPropertyBag& InParameters)
 {
 	UPropertyAnimatorCoreTimeSourceBase* TimeSource = GetActiveTimeSource();
 
@@ -431,13 +429,11 @@ void UPropertyAnimatorCoreBase::EvaluateAnimator()
 	SaveProperties();
 
 	EvaluatedPropertyValues.Reset();
-
-	FPropertyAnimatorCoreEvaluationParameters Parameters;
-	Parameters.AnimatorsMagnitude = GetAnimatorComponentMagnitude();
-	Parameters.TimeElapsed = TimeElapsed.GetValue();
+	InParameters.AddProperty(TimeElapsedParameterName, EPropertyBagPropertyType::Double);
+	InParameters.SetValueDouble(TimeElapsedParameterName, TimeElapsed.GetValue());
 
 	bEvaluatingProperties = true;
-	EvaluateProperties(Parameters);
+	EvaluateProperties(InParameters);
 	bEvaluatingProperties = false;
 }
 
@@ -530,18 +526,13 @@ void UPropertyAnimatorCoreBase::SaveProperties()
 	}, bResolve);
 }
 
-TArray<FString> UPropertyAnimatorCoreBase::GetTimeSourceNames() const
+TArray<FName> UPropertyAnimatorCoreBase::GetTimeSourceNames() const
 {
-	TArray<FString> TimeSourceNames;
+	TArray<FName> TimeSourceNames;
 
-	if (const UPropertyAnimatorCoreSubsystem* ControllerSubsystem = UPropertyAnimatorCoreSubsystem::Get())
+	if (const UPropertyAnimatorCoreSubsystem* AnimatorSubsystem = UPropertyAnimatorCoreSubsystem::Get())
 	{
-		const TArray<FName> TimeSources = ControllerSubsystem->GetTimeSourceNames();
-
-		Algo::Transform(TimeSources, TimeSourceNames, [](const FName& InTimeSource)
-		{
-			return InTimeSource.ToString();
-		});
+		TimeSourceNames = AnimatorSubsystem->GetTimeSourceNames();
 	}
 
 	return TimeSourceNames;
@@ -563,11 +554,12 @@ UPropertyAnimatorCoreTimeSourceBase* UPropertyAnimatorCoreBase::FindOrAddTimeSou
 
 	// Check cached time source instances
 	UPropertyAnimatorCoreTimeSourceBase* NewTimeSource = nullptr;
-	if (TObjectPtr<UPropertyAnimatorCoreTimeSourceBase> const* TimeSourceInstance = TimeSourcesInstances.Find(InTimeSourceName))
+
+	for (const TObjectPtr<UPropertyAnimatorCoreTimeSourceBase>& TimeSource : TimeSources)
 	{
-		if (const TObjectPtr<UPropertyAnimatorCoreTimeSourceBase> CachedTimeSource = *TimeSourceInstance)
+		if (TimeSource && TimeSource->GetTimeSourceName() == InTimeSourceName)
 		{
-			NewTimeSource = CachedTimeSource.Get();
+			NewTimeSource = TimeSource.Get();
 		}
 	}
 
@@ -575,9 +567,10 @@ UPropertyAnimatorCoreTimeSourceBase* UPropertyAnimatorCoreBase::FindOrAddTimeSou
 	if (!NewTimeSource)
 	{
 		NewTimeSource = Subsystem->CreateNewTimeSource(InTimeSourceName, this);
+
 		if (NewTimeSource)
 		{
-			TimeSourcesInstances.Add(InTimeSourceName, NewTimeSource);
+			TimeSources.Add(NewTimeSource);
 		}
 	}
 
@@ -632,7 +625,9 @@ bool UPropertyAnimatorCoreBase::LinkProperty(const FPropertyAnimatorCoreData& In
 		return false;
 	}
 
-	if (!IsPropertySupported(InLinkProperty))
+	const EPropertyAnimatorPropertySupport Support = GetPropertySupport(InLinkProperty);
+
+	if (Support == EPropertyAnimatorPropertySupport::None)
 	{
 		return false;
 	}
@@ -654,7 +649,7 @@ bool UPropertyAnimatorCoreBase::LinkProperty(const FPropertyAnimatorCoreData& In
 	PropertyContext->ConstructInternal(InLinkProperty);
 
 	LinkedProperties.Add(PropertyContext);
-	OnPropertyLinked(PropertyContext);
+	OnPropertyLinked(PropertyContext, Support);
 
 	UPropertyAnimatorCoreBase::OnAnimatorPropertyLinkedDelegate.Broadcast(this, InLinkProperty);
 
@@ -683,14 +678,25 @@ bool UPropertyAnimatorCoreBase::IsPropertyLinked(const FPropertyAnimatorCoreData
 {
 	return LinkedProperties.ContainsByPredicate([&InPropertyData](const UPropertyAnimatorCoreContext* InOptions)
 	{
-		return InOptions && InOptions->GetAnimatedProperty() == InPropertyData;
+		return InOptions
+			&& (
+				InOptions->GetAnimatedProperty() == InPropertyData
+				|| InOptions->GetAnimatedProperty().IsOwning(InPropertyData)
+			);
 	});
 }
 
 bool UPropertyAnimatorCoreBase::IsPropertiesLinked(const TSet<FPropertyAnimatorCoreData>& InProperties) const
 {
-	const TSet<FPropertyAnimatorCoreData> LinkedPropertiesSet = GetLinkedProperties();
-	return LinkedPropertiesSet.Includes(InProperties);
+	for (const FPropertyAnimatorCoreData& Property : InProperties)
+	{
+		if (!IsPropertyLinked(Property))
+		{
+			return false;
+		}
+	}
+
+	return !InProperties.IsEmpty();
 }
 
 TSet<FPropertyAnimatorCoreData> UPropertyAnimatorCoreBase::GetInnerPropertiesLinked(const FPropertyAnimatorCoreData& InPropertyData) const
diff --git a/Engine/Plugins/Experimental/PropertyAnimatorCore/Source/PropertyAnimatorCore/Private/Components/PropertyAnimatorCoreComponent.cpp b/Engine/Plugins/Experimental/PropertyAnimatorCore/Source/PropertyAnimatorCore/Private/Components/PropertyAnimatorCoreComponent.cpp
--- a/Engine/Plugins/Experimental/PropertyAnimatorCore/Source/PropertyAnimatorCore/Private/Components/PropertyAnimatorCoreComponent.cpp
+++ b/Engine/Plugins/Experimental/PropertyAnimatorCore/Source/PropertyAnimatorCore/Private/Components/PropertyAnimatorCoreComponent.cpp
@@ -208,6 +208,8 @@ void UPropertyAnimatorCoreComponent::TickComponent(float InDeltaTime, ELevelTick
 		return;
 	}
 
+	FInstancedPropertyBag Parameters;
+
 	for (const TObjectPtr<UPropertyAnimatorCoreBase>& Animator : PropertyAnimators)
 	{
 		if (!IsValid(Animator))
@@ -215,7 +217,11 @@ void UPropertyAnimatorCoreComponent::TickComponent(float InDeltaTime, ELevelTick
 			continue;
 		}
 
-		Animator->EvaluateAnimator();
+		Parameters.Reset();
+		Parameters.AddProperty(UPropertyAnimatorCoreBase::MagnitudeParameterName, EPropertyBagPropertyType::Float);
+		Parameters.SetValueFloat(UPropertyAnimatorCoreBase::MagnitudeParameterName, AnimatorsMagnitude);
+
+		Animator->EvaluateAnimator(Parameters);
 	}
 }
 
diff --git a/Engine/Plugins/Experimental/PropertyAnimatorCore/Source/PropertyAnimatorCore/Private/Presets/PropertyAnimatorCorePresetBase.cpp b/Engine/Plugins/Experimental/PropertyAnimatorCore/Source/PropertyAnimatorCore/Private/Presets/PropertyAnimatorCorePresetBase.cpp
--- a/Engine/Plugins/Experimental/PropertyAnimatorCore/Source/PropertyAnimatorCore/Private/Presets/PropertyAnimatorCorePresetBase.cpp
+++ b/Engine/Plugins/Experimental/PropertyAnimatorCore/Source/PropertyAnimatorCore/Private/Presets/PropertyAnimatorCorePresetBase.cpp
@@ -22,7 +22,7 @@ void UPropertyAnimatorCorePresetBase::GetSupportedPresetProperties(const AActor*
 
 	for (const FPropertyAnimatorCoreData& PresetProperty : PresetProperties)
 	{
-		InAnimator->GetPropertiesSupported(PresetProperty, OutProperties, false);
+		InAnimator->GetPropertiesSupported(PresetProperty, OutProperties, /** SearchDepth */3);
 	}
 }
 
diff --git a/Engine/Plugins/Experimental/PropertyAnimatorCore/Source/PropertyAnimatorCore/Private/Properties/Handlers/PropertyAnimatorCoreRotatorHandler.cpp b/Engine/Plugins/Experimental/PropertyAnimatorCore/Source/PropertyAnimatorCore/Private/Properties/Handlers/PropertyAnimatorCoreRotatorHandler.cpp
new file mode 100644
--- /dev/null
+++ b/Engine/Plugins/Experimental/PropertyAnimatorCore/Source/PropertyAnimatorCore/Private/Properties/Handlers/PropertyAnimatorCoreRotatorHandler.cpp
@@ -0,0 +1,82 @@
+// Copyright Epic Games, Inc. All Rights Reserved.
+
+#include "Properties/Handlers/PropertyAnimatorCoreRotatorHandler.h"
+
+bool UPropertyAnimatorCoreRotatorHandler::IsPropertySupported(const FPropertyAnimatorCoreData& InPropertyData) const
+{
+	if (InPropertyData.IsA<FStructProperty>() && InPropertyData.GetLeafPropertyTypeName() == NAME_Rotator)
+	{
+		return true;
+	}
+
+	return Super::IsPropertySupported(InPropertyData);
+}
+
+bool UPropertyAnimatorCoreRotatorHandler::GetValue(const FPropertyAnimatorCoreData& InPropertyData, FInstancedPropertyBag& OutValue)
+{
+	const FName PropertyName(InPropertyData.GetPathHash());
+	OutValue.AddProperty(PropertyName, EPropertyBagPropertyType::Struct, TBaseStructure<FRotator>::Get());
+
+	FRotator Value;
+	InPropertyData.GetPropertyValuePtr(&Value);
+
+	OutValue.SetValueStruct(PropertyName, Value);
+
+	return true;
+}
+
+bool UPropertyAnimatorCoreRotatorHandler::SetValue(const FPropertyAnimatorCoreData& InPropertyData, const FInstancedPropertyBag& InValue)
+{
+	const FName PropertyName(InPropertyData.GetPathHash());
+	TValueOrError<FRotator*, EPropertyBagResult> ValueResult = InValue.GetValueStruct<FRotator>(PropertyName);
+	if (!ValueResult.HasValue())
+	{
+		return false;
+	}
+
+	FRotator NewValue = *ValueResult.GetValue();
+	InPropertyData.SetPropertyValuePtr(&NewValue);
+
+	return true;
+}
+
+bool UPropertyAnimatorCoreRotatorHandler::IsAdditiveSupported() const
+{
+	return true;
+}
+
+bool UPropertyAnimatorCoreRotatorHandler::AddValue(const FPropertyAnimatorCoreData& InPropertyData, const FInstancedPropertyBag& InValue)
+{
+	const FName PropertyName(InPropertyData.GetPathHash());
+	const TValueOrError<FRotator*, EPropertyBagResult> ValueResult = InValue.GetValueStruct<FRotator>(PropertyName);
+	if (!ValueResult.HasValue())
+	{
+		return false;
+	}
+
+	FRotator Value;
+	InPropertyData.GetPropertyValuePtr(&Value);
+
+	FRotator NewValue = Value + *ValueResult.GetValue();
+	InPropertyData.SetPropertyValuePtr(&NewValue);
+
+	return true;
+}
+
+bool UPropertyAnimatorCoreRotatorHandler::SubtractValue(const FPropertyAnimatorCoreData& InPropertyData, const FInstancedPropertyBag& InValue)
+{
+	const FName PropertyName(InPropertyData.GetPathHash());
+	const TValueOrError<FRotator*, EPropertyBagResult> ValueResult = InValue.GetValueStruct<FRotator>(PropertyName);
+	if (!ValueResult.HasValue())
+	{
+		return false;
+	}
+
+	FRotator Value;
+	InPropertyData.GetPropertyValuePtr(&Value);
+
+	FRotator NewValue = Value - *ValueResult.GetValue();
+	InPropertyData.SetPropertyValuePtr(&NewValue);
+
+	return true;
+}
diff --git a/Engine/Plugins/Experimental/PropertyAnimatorCore/Source/PropertyAnimatorCore/Private/Properties/Handlers/PropertyAnimatorCoreRotatorHandler.h b/Engine/Plugins/Experimental/PropertyAnimatorCore/Source/PropertyAnimatorCore/Private/Properties/Handlers/PropertyAnimatorCoreRotatorHandler.h
new file mode 100644
--- /dev/null
+++ b/Engine/Plugins/Experimental/PropertyAnimatorCore/Source/PropertyAnimatorCore/Private/Properties/Handlers/PropertyAnimatorCoreRotatorHandler.h
@@ -0,0 +1,23 @@
+// Copyright Epic Games, Inc. All Rights Reserved.
+
+#pragma once
+
+#include "Properties/Handlers/PropertyAnimatorCoreHandlerBase.h"
+#include "PropertyAnimatorCoreRotatorHandler.generated.h"
+
+/** Handles any rotator property values, get and set handler */
+UCLASS(Transient)
+class UPropertyAnimatorCoreRotatorHandler : public UPropertyAnimatorCoreHandlerBase
+{
+	GENERATED_BODY()
+
+public:
+	//~ Begin UPropertyAnimatorCoreHandlerBase
+	virtual bool IsPropertySupported(const FPropertyAnimatorCoreData& InPropertyData) const override;
+	virtual bool GetValue(const FPropertyAnimatorCoreData& InPropertyData, FInstancedPropertyBag& OutValue) override;
+	virtual bool SetValue(const FPropertyAnimatorCoreData& InPropertyData, const FInstancedPropertyBag& InValue) override;
+	virtual bool IsAdditiveSupported() const override;
+	virtual bool AddValue(const FPropertyAnimatorCoreData& InPropertyData, const FInstancedPropertyBag& InValue) override;
+	virtual bool SubtractValue(const FPropertyAnimatorCoreData& InPropertyData, const FInstancedPropertyBag& InValue) override;
+	//~ End UPropertyAnimatorCoreHandlerBase
+};
\ No newline at end of file
diff --git a/Engine/Plugins/Experimental/PropertyAnimatorCore/Source/PropertyAnimatorCore/Private/Properties/Handlers/PropertyAnimatorCoreVectorHandler.cpp b/Engine/Plugins/Experimental/PropertyAnimatorCore/Source/PropertyAnimatorCore/Private/Properties/Handlers/PropertyAnimatorCoreVectorHandler.cpp
new file mode 100644
--- /dev/null
+++ b/Engine/Plugins/Experimental/PropertyAnimatorCore/Source/PropertyAnimatorCore/Private/Properties/Handlers/PropertyAnimatorCoreVectorHandler.cpp
@@ -0,0 +1,82 @@
+// Copyright Epic Games, Inc. All Rights Reserved.
+
+#include "Properties/Handlers/PropertyAnimatorCoreVectorHandler.h"
+
+bool UPropertyAnimatorCoreVectorHandler::IsPropertySupported(const FPropertyAnimatorCoreData& InPropertyData) const
+{
+	if (InPropertyData.IsA<FStructProperty>() && InPropertyData.GetLeafPropertyTypeName() == NAME_Vector)
+	{
+		return true;
+	}
+
+	return Super::IsPropertySupported(InPropertyData);
+}
+
+bool UPropertyAnimatorCoreVectorHandler::GetValue(const FPropertyAnimatorCoreData& InPropertyData, FInstancedPropertyBag& OutValue)
+{
+	const FName PropertyName(InPropertyData.GetPathHash());
+	OutValue.AddProperty(PropertyName, EPropertyBagPropertyType::Struct, TBaseStructure<FVector>::Get());
+
+	FVector Value;
+	InPropertyData.GetPropertyValuePtr(&Value);
+
+	OutValue.SetValueStruct(PropertyName, Value);
+
+	return true;
+}
+
+bool UPropertyAnimatorCoreVectorHandler::SetValue(const FPropertyAnimatorCoreData& InPropertyData, const FInstancedPropertyBag& InValue)
+{
+	const FName PropertyName(InPropertyData.GetPathHash());
+	TValueOrError<FVector*, EPropertyBagResult> ValueResult = InValue.GetValueStruct<FVector>(PropertyName);
+	if (!ValueResult.HasValue())
+	{
+		return false;
+	}
+
+	FVector NewValue = *ValueResult.GetValue();
+	InPropertyData.SetPropertyValuePtr(&NewValue);
+
+	return true;
+}
+
+bool UPropertyAnimatorCoreVectorHandler::IsAdditiveSupported() const
+{
+	return true;
+}
+
+bool UPropertyAnimatorCoreVectorHandler::AddValue(const FPropertyAnimatorCoreData& InPropertyData, const FInstancedPropertyBag& InValue)
+{
+	const FName PropertyName(InPropertyData.GetPathHash());
+	const TValueOrError<FVector*, EPropertyBagResult> ValueResult = InValue.GetValueStruct<FVector>(PropertyName);
+	if (!ValueResult.HasValue())
+	{
+		return false;
+	}
+
+	FVector Value;
+	InPropertyData.GetPropertyValuePtr(&Value);
+
+	FVector NewValue = Value + *ValueResult.GetValue();
+	InPropertyData.SetPropertyValuePtr(&NewValue);
+
+	return true;
+}
+
+bool UPropertyAnimatorCoreVectorHandler::SubtractValue(const FPropertyAnimatorCoreData& InPropertyData, const FInstancedPropertyBag& InValue)
+{
+	const FName PropertyName(InPropertyData.GetPathHash());
+	const TValueOrError<FVector*, EPropertyBagResult> ValueResult = InValue.GetValueStruct<FVector>(PropertyName);
+	if (!ValueResult.HasValue())
+	{
+		return false;
+	}
+
+	FVector Value;
+	InPropertyData.GetPropertyValuePtr(&Value);
+
+	FVector NewValue = Value - *ValueResult.GetValue();
+	InPropertyData.SetPropertyValuePtr(&NewValue);
+
+	return true;
+}
diff --git a/Engine/Plugins/Experimental/PropertyAnimatorCore/Source/PropertyAnimatorCore/Private/Properties/Handlers/PropertyAnimatorCoreVectorHandler.h b/Engine/Plugins/Experimental/PropertyAnimatorCore/Source/PropertyAnimatorCore/Private/Properties/Handlers/PropertyAnimatorCoreVectorHandler.h
new file mode 100644
--- /dev/null
+++ b/Engine/Plugins/Experimental/PropertyAnimatorCore/Source/PropertyAnimatorCore/Private/Properties/Handlers/PropertyAnimatorCoreVectorHandler.h
@@ -0,0 +1,23 @@
+// Copyright Epic Games, Inc. All Rights Reserved.
+
+#pragma once
+
+#include "Properties/Handlers/PropertyAnimatorCoreHandlerBase.h"
+#include "PropertyAnimatorCoreVectorHandler.generated.h"
+
+/** Handles any vector 3d property values, get and set handler */
+UCLASS(Transient)
+class UPropertyAnimatorCoreVectorHandler : public UPropertyAnimatorCoreHandlerBase
+{
+	GENERATED_BODY()
+
+public:
+	//~ Begin UPropertyAnimatorCoreHandlerBase
+	virtual bool IsPropertySupported(const FPropertyAnimatorCoreData& InPropertyData) const override;
+	virtual bool GetValue(const FPropertyAnimatorCoreData& InPropertyData, FInstancedPropertyBag& OutValue) override;
+	virtual bool SetValue(const FPropertyAnimatorCoreData& InPropertyData, const FInstancedPropertyBag& InValue) override;
+	virtual bool IsAdditiveSupported() const override;
+	virtual bool AddValue(const FPropertyAnimatorCoreData& InPropertyData, const FInstancedPropertyBag& InValue) override;
+	virtual bool SubtractValue(const FPropertyAnimatorCoreData& InPropertyData, const FInstancedPropertyBag& InValue) override;
+	//~ End UPropertyAnimatorCoreHandlerBase
+};
\ No newline at end of file
diff --git a/Engine/Plugins/Experimental/PropertyAnimatorCore/Source/PropertyAnimatorCore/Private/Properties/PropertyAnimatorCoreContext.cpp b/Engine/Plugins/Experimental/PropertyAnimatorCore/Source/PropertyAnimatorCore/Private/Properties/PropertyAnimatorCoreContext.cpp
--- a/Engine/Plugins/Experimental/PropertyAnimatorCore/Source/PropertyAnimatorCore/Private/Properties/PropertyAnimatorCoreContext.cpp
+++ b/Engine/Plugins/Experimental/PropertyAnimatorCore/Source/PropertyAnimatorCore/Private/Properties/PropertyAnimatorCoreContext.cpp
@@ -85,6 +85,16 @@ void UPropertyAnimatorCoreContext::SetAnimated(bool bInAnimated)
 	OnAnimatedChanged();
 }
 
+void UPropertyAnimatorCoreContext::SetMagnitude(float InMagnitude)
+{
+	Magnitude = FMath::Clamp(InMagnitude, 0.f, 1.f);
+}
+
+void UPropertyAnimatorCoreContext::SetTimeOffset(double InTimeOffset)
+{
+	TimeOffset = InTimeOffset;
+}
+
 void UPropertyAnimatorCoreContext::SetMode(EPropertyAnimatorCoreMode InMode)
 {
 	if (InMode == Mode)
@@ -433,7 +443,7 @@ void UPropertyAnimatorCoreContext::OnAnimatedChanged()
 	}
 }
 
-void UPropertyAnimatorCoreContext::SetEvaluationResult(const FPropertyAnimatorCoreData& InResolvedProperty, const FInstancedPropertyBag& InEvaluatedValues)
+void UPropertyAnimatorCoreContext::CommitEvaluationResult(const FPropertyAnimatorCoreData& InResolvedProperty, const FInstancedPropertyBag& InEvaluatedValues)
 {
 	if (!IsAnimated())
 	{
diff --git a/Engine/Plugins/Experimental/PropertyAnimatorCore/Source/PropertyAnimatorCore/Private/Properties/PropertyAnimatorCoreData.cpp b/Engine/Plugins/Experimental/PropertyAnimatorCore/Source/PropertyAnimatorCore/Private/Properties/PropertyAnimatorCoreData.cpp
--- a/Engine/Plugins/Experimental/PropertyAnimatorCore/Source/PropertyAnimatorCore/Private/Properties/PropertyAnimatorCoreData.cpp
+++ b/Engine/Plugins/Experimental/PropertyAnimatorCore/Source/PropertyAnimatorCore/Private/Properties/PropertyAnimatorCoreData.cpp
@@ -176,12 +176,22 @@ FName FPropertyAnimatorCoreData::GetMemberPropertyName() const
 	return MemberProperty ? MemberProperty->GetFName() : NAME_None;
 }
 
+FName FPropertyAnimatorCoreData::GetMemberPropertyTypeName() const
+{
+	return GetPropertyTypeName(GetMemberProperty());
+}
+
 FName FPropertyAnimatorCoreData::GetLeafPropertyName() const
 {
 	const FProperty* LeafProperty = GetLeafProperty();
 	return LeafProperty ? LeafProperty->GetFName() : NAME_None;
 }
 
+FName FPropertyAnimatorCoreData::GetLeafPropertyTypeName() const
+{
+	return GetPropertyTypeName(GetLeafProperty());
+}
+
 TArray<FProperty*> FPropertyAnimatorCoreData::GetChainProperties() const
 {
 	TArray<FProperty*> Properties;
@@ -503,6 +513,21 @@ void FPropertyAnimatorCoreData::CopyPropertyValue(const FProperty* InProperty, c
 	}
 }
 
+FName FPropertyAnimatorCoreData::GetPropertyTypeName(const FProperty* InProperty)
+{
+	if (InProperty)
+	{
+		if (const FStructProperty* StructProperty = CastField<FStructProperty>(InProperty))
+		{
+			return StructProperty->Struct->GetFName();
+		}
+
+		return FName(InProperty->GetCPPType());
+	}
+
+	return NAME_None;
+}
+
 void FPropertyAnimatorCoreData::GeneratePropertyPath()
 {
 	const UObject* Owner = GetOwner();
diff --git a/Engine/Plugins/Experimental/PropertyAnimatorCore/Source/PropertyAnimatorCore/Private/Subsystems/PropertyAnimatorCoreSubsystem.cpp b/Engine/Plugins/Experimental/PropertyAnimatorCore/Source/PropertyAnimatorCore/Private/Subsystems/PropertyAnimatorCoreSubsystem.cpp
--- a/Engine/Plugins/Experimental/PropertyAnimatorCore/Source/PropertyAnimatorCore/Private/Subsystems/PropertyAnimatorCoreSubsystem.cpp
+++ b/Engine/Plugins/Experimental/PropertyAnimatorCore/Source/PropertyAnimatorCore/Private/Subsystems/PropertyAnimatorCoreSubsystem.cpp
@@ -159,22 +159,26 @@ bool UPropertyAnimatorCoreSubsystem::IsPropertySupported(const FPropertyAnimator
 		return false;
 	}
 
-	for (const TWeakObjectPtr<UPropertyAnimatorCoreBase>& Controller : AnimatorsWeak)
+	for (const TWeakObjectPtr<UPropertyAnimatorCoreBase>& AnimatorWeak : AnimatorsWeak)
 	{
+		const UPropertyAnimatorCoreBase* Animator = AnimatorWeak.Get();
+
+		if (!Animator)
+		{
+			continue;
+		}
+
 		if (bInCheckNestedProperties)
 		{
 			TSet<FPropertyAnimatorCoreData> OutProperties;
-			if (Controller->GetPropertiesSupported(InPropertyData, OutProperties))
+			if (Animator->GetPropertiesSupported(InPropertyData, OutProperties))
 			{
 				return true;
 			}
 		}
-		else
+		else if (Animator->HasPropertySupport(InPropertyData))
 		{
-			if (Controller->IsPropertySupported(InPropertyData))
-			{
-				return true;
-			}
+			return true;
 		}
 	}
 
diff --git a/Engine/Plugins/Experimental/PropertyAnimatorCore/Source/PropertyAnimatorCore/Private/TimeSources/PropertyAnimatorCoreTimeSourceBase.cpp b/Engine/Plugins/Experimental/PropertyAnimatorCore/Source/PropertyAnimatorCore/Private/TimeSources/PropertyAnimatorCoreTimeSourceBase.cpp
--- a/Engine/Plugins/Experimental/PropertyAnimatorCore/Source/PropertyAnimatorCore/Private/TimeSources/PropertyAnimatorCoreTimeSourceBase.cpp
+++ b/Engine/Plugins/Experimental/PropertyAnimatorCore/Source/PropertyAnimatorCore/Private/TimeSources/PropertyAnimatorCoreTimeSourceBase.cpp
@@ -4,6 +4,13 @@
 
 #include "Animators/PropertyAnimatorCoreBase.h"
 
+#if WITH_EDITOR
+FName UPropertyAnimatorCoreTimeSourceBase::GetTimeElapsedPropertyName()
+{
+	return GET_MEMBER_NAME_CHECKED(UPropertyAnimatorCoreTimeSourceBase, TimeElapsed);
+}
+#endif
+
 void UPropertyAnimatorCoreTimeSourceBase::ActivateTimeSource()
 {
 	if (IsTimeSourceActive())
diff --git a/Engine/Plugins/Experimental/PropertyAnimatorCore/Source/PropertyAnimatorCore/PropertyAnimatorCore.Build.cs b/Engine/Plugins/Experimental/PropertyAnimatorCore/Source/PropertyAnimatorCore/PropertyAnimatorCore.Build.cs
--- a/Engine/Plugins/Experimental/PropertyAnimatorCore/Source/PropertyAnimatorCore/PropertyAnimatorCore.Build.cs
+++ b/Engine/Plugins/Experimental/PropertyAnimatorCore/Source/PropertyAnimatorCore/PropertyAnimatorCore.Build.cs
@@ -12,7 +12,7 @@ public class PropertyAnimatorCore : ModuleRules
 			new string[]
 			{
 				"Core",
-				"CoreUObject",
+				"CoreUObject"
 			});
 
 		PrivateDependencyModuleNames.AddRange(
diff --git a/Engine/Plugins/Experimental/PropertyAnimatorCore/Source/PropertyAnimatorCore/Public/Animators/PropertyAnimatorCoreBase.h b/Engine/Plugins/Experimental/PropertyAnimatorCore/Source/PropertyAnimatorCore/Public/Animators/PropertyAnimatorCoreBase.h
--- a/Engine/Plugins/Experimental/PropertyAnimatorCore/Source/PropertyAnimatorCore/Public/Animators/PropertyAnimatorCoreBase.h
+++ b/Engine/Plugins/Experimental/PropertyAnimatorCore/Source/PropertyAnimatorCore/Public/Animators/PropertyAnimatorCoreBase.h
@@ -11,14 +11,13 @@ class UPropertyAnimatorCoreComponent;
 class UPropertyAnimatorCoreGroupBase;
 class UPropertyAnimatorCoreTimeSourceBase;
 
-/** Parameters to use for the current evaluation results */
-struct FPropertyAnimatorCoreEvaluationParameters
+UENUM(BlueprintType)
+enum class EPropertyAnimatorPropertySupport : uint8
 {
-	/** Time elapsed generated by the active time source */
-	double TimeElapsed = 0.f;
-
-	/** Global magnitude for all animators */
-	float AnimatorsMagnitude = 1.f;
+	None = 0,
+	Incomplete = 1 << 0,
+	Complete = 1 << 1,
+	All = Incomplete | Complete
 };
 
 /** Abstract base class for any Animator, holds a set of linked properties */
@@ -52,6 +51,11 @@ public:
 	/** Called when a property is unlinked to a Animator */
 	PROPERTYANIMATORCORE_API static FOnAnimatorPropertyUpdated OnAnimatorPropertyUnlinkedDelegate;
 
+	static constexpr const TCHAR* TimeElapsedParameterName = TEXT("TimeElapsed");
+	static constexpr const TCHAR* MagnitudeParameterName = TEXT("Magnitude");
+	static constexpr const TCHAR* FrequencyParameterName = TEXT("Frequency");
+	static constexpr const TCHAR* AlphaParameterName = TEXT("Alpha");
+
 	PROPERTYANIMATORCORE_API UPropertyAnimatorCoreBase();
 
 	PROPERTYANIMATORCORE_API AActor* GetAnimatorActor() const;
@@ -113,25 +117,21 @@ public:
 	PROPERTYANIMATORCORE_API TSet<FPropertyAnimatorCoreData> GetInnerPropertiesLinked(const FPropertyAnimatorCoreData& InPropertyData) const;
 
 	/**
-	 * Checks recursively for properties inside that member property that are supported by this Animator
-	 * if the member property is directly supported only returns that one, calls IsPropertySupported to check
-	 * When recursive is false, it will stop if member property or children is controllable and not look further in the chain
+	 * Checks recursively for properties that are supported by this Animator, calls IsPropertySupported to check
+	 * Stops when the InSearchDepth has been reached otherwise continues to gather supported properties
 	 */
-	PROPERTYANIMATORCORE_API bool GetPropertiesSupported(const FPropertyAnimatorCoreData& InPropertyData, TSet<FPropertyAnimatorCoreData>& OutProperties, bool bInRecursiveSearch = false) const;
+	PROPERTYANIMATORCORE_API bool GetPropertiesSupported(const FPropertyAnimatorCoreData& InPropertyData, TSet<FPropertyAnimatorCoreData>& OutProperties, uint8 InSearchDepth = 1, EPropertyAnimatorPropertySupport InSupportExpected = EPropertyAnimatorPropertySupport::All) const;
 
-	/** Does not recurse, checks if this Animator supports this specific property, use GetPropertiesSupported to check if nested properties are supported */
-	PROPERTYANIMATORCORE_API bool IsPropertySupported(const FPropertyAnimatorCoreData& InPropertyData) const;
+	/** Retrieves the support level of a property */
+	PROPERTYANIMATORCORE_API EPropertyAnimatorPropertySupport GetPropertySupport(const FPropertyAnimatorCoreData& InPropertyData) const;
 
-	/** Override this to check if a property is directly supported by this animator */
-	virtual bool IsPropertyDirectlySupported(const FPropertyAnimatorCoreData& InPropertyData) const
-	{
-		return false;
-	}
+	/** Checks if a property support is available */
+	PROPERTYANIMATORCORE_API bool HasPropertySupport(const FPropertyAnimatorCoreData& InPropertyData, EPropertyAnimatorPropertySupport InSupportExpected = EPropertyAnimatorPropertySupport::All) const;
 
-	/** Override this to check if a property is indirectly supported by this animator with conversions */
-	virtual bool IsPropertyIndirectlySupported(const FPropertyAnimatorCoreData& InPropertyData) const
+	/** Override this to check if a property is supported by this animator */
+	virtual EPropertyAnimatorPropertySupport IsPropertySupported(const FPropertyAnimatorCoreData& InPropertyData) const
 	{
-		return false;
+		return EPropertyAnimatorPropertySupport::None;
 	}
 
 	/** Get the context for the linked property */
@@ -229,7 +229,7 @@ protected:
 
 					if (InFunction(PropertyContext, ResolvedPropertyData, EvaluatedPropertyValues))
 					{
-						PropertyContext->SetEvaluationResult(ResolvedPropertyData, EvaluatedPropertyValues);
+						PropertyContext->CommitEvaluationResult(ResolvedPropertyData, EvaluatedPropertyValues);
 					}
 				}
 			}
@@ -249,11 +249,11 @@ protected:
 	/** Returns the property context class to use */
 	PROPERTYANIMATORCORE_API virtual TSubclassOf<UPropertyAnimatorCoreContext> GetPropertyContextClass(const FPropertyAnimatorCoreData& InProperty);
 
-	virtual void OnPropertyLinked(UPropertyAnimatorCoreContext* InLinkedProperty) {}
+	virtual void OnPropertyLinked(UPropertyAnimatorCoreContext* InLinkedProperty, EPropertyAnimatorPropertySupport InSupport) {}
 	virtual void OnPropertyUnlinked(UPropertyAnimatorCoreContext* InUnlinkedProperty) {}
 
 	/** Apply animators effect on linked properties */
-	virtual void EvaluateProperties(const FPropertyAnimatorCoreEvaluationParameters& InParameters) {}
+	virtual void EvaluateProperties(FInstancedPropertyBag& InParameters) {}
 
 private:
 	/** Restore modified properties to original state */
@@ -263,7 +263,7 @@ private:
 	void SaveProperties();
 
 	/** Called by the component to evaluate this animator */
-	void EvaluateAnimator();
+	void EvaluateAnimator(FInstancedPropertyBag& InParameters);
 
 	void OnObjectReplaced(const TMap<UObject*, UObject*>& InReplacementMap);
 
@@ -273,9 +273,6 @@ private:
 
 	void CleanLinkedProperties();
 
-	/** Clean invalid or renamed time sources */
-	void CleanTimeSources();
-
 	void OnTimeSourceNameChanged();
 
 	/** Called after an action that causes the owner to change */
@@ -284,7 +281,7 @@ private:
 	UPropertyAnimatorCoreTimeSourceBase* FindOrAddTimeSource(FName InTimeSourceName);
 
 	UFUNCTION()
-	TArray<FString> GetTimeSourceNames() const;
+	TArray<FName> GetTimeSourceNames() const;
 
 	/** Enable control of properties linked to this Animator */
 	UPROPERTY(EditInstanceOnly, Getter="GetAnimatorEnabled", Setter="SetAnimatorEnabled", Category="Animator", meta=(DisplayPriority="0", AllowPrivateAccess="true"))
@@ -311,9 +308,14 @@ private:
 	TObjectPtr<UPropertyAnimatorCoreTimeSourceBase> ActiveTimeSource;
 
 	/** The cached time source used by this Animator */
-	UPROPERTY()
+	UE_DEPRECATED(5.5, "Use TimeSources instead")
+	UPROPERTY(meta=(DeprecatedProperty, DeprecationMessage="Use TimeSources instead"))
 	TMap<FName, TObjectPtr<UPropertyAnimatorCoreTimeSourceBase>> TimeSourcesInstances;
 
+	/** Cached time sources used by this animator */
+	UPROPERTY()
+	TArray<TObjectPtr<UPropertyAnimatorCoreTimeSourceBase>> TimeSources;
+
 	/** Evaluated property container, reset on every update round */
 	UPROPERTY(Transient, DuplicateTransient, NonTransactional)
 	FInstancedPropertyBag EvaluatedPropertyValues;
diff --git a/Engine/Plugins/Experimental/PropertyAnimatorCore/Source/PropertyAnimatorCore/Public/Components/PropertyAnimatorCoreComponent.h b/Engine/Plugins/Experimental/PropertyAnimatorCore/Source/PropertyAnimatorCore/Public/Components/PropertyAnimatorCoreComponent.h
--- a/Engine/Plugins/Experimental/PropertyAnimatorCore/Source/PropertyAnimatorCore/Public/Components/PropertyAnimatorCoreComponent.h
+++ b/Engine/Plugins/Experimental/PropertyAnimatorCore/Source/PropertyAnimatorCore/Public/Components/PropertyAnimatorCoreComponent.h
@@ -115,4 +115,8 @@ private:
 	/** Transient copy of property animators when changes are detected to see the diff only */
 	UPROPERTY(Transient, DuplicateTransient, TextExportTransient)
 	TArray<TObjectPtr<UPropertyAnimatorCoreBase>> PropertyAnimatorsInternal;
+
+	/** Cached time sources used by this animator component */
+	UPROPERTY()
+	TArray<TObjectPtr<UPropertyAnimatorCoreTimeSourceBase>> TimeSourceInstances;
 };
\ No newline at end of file
diff --git a/Engine/Plugins/Experimental/PropertyAnimatorCore/Source/PropertyAnimatorCore/Public/Properties/PropertyAnimatorCoreContext.h b/Engine/Plugins/Experimental/PropertyAnimatorCore/Source/PropertyAnimatorCore/Public/Properties/PropertyAnimatorCoreContext.h
--- a/Engine/Plugins/Experimental/PropertyAnimatorCore/Source/PropertyAnimatorCore/Public/Properties/PropertyAnimatorCoreContext.h
+++ b/Engine/Plugins/Experimental/PropertyAnimatorCore/Source/PropertyAnimatorCore/Public/Properties/PropertyAnimatorCoreContext.h
@@ -53,6 +53,18 @@ public:
 		return bAnimated;
 	}
 
+	PROPERTYANIMATORCORE_API void SetMagnitude(float InMagnitude);
+	float GetMagnitude() const
+	{
+		return Magnitude;
+	}
+
+	PROPERTYANIMATORCORE_API void SetTimeOffset(double InTimeOffset);
+	double GetTimeOffset() const
+	{
+		return TimeOffset;
+	}
+
 	PROPERTYANIMATORCORE_API void SetMode(EPropertyAnimatorCoreMode InMode);
 	EPropertyAnimatorCoreMode GetMode() const
 	{
@@ -82,6 +94,12 @@ public:
 	/** Called when the owner has changed and we want to update the animated property */
 	bool ResolvePropertyOwner(AActor* InNewOwner);
 
+	/** Evaluates a property within this context based on animator result */
+	virtual bool EvaluateProperty(const FPropertyAnimatorCoreData& InProperty, const FInstancedPropertyBag& InAnimatorResult, FInstancedPropertyBag& OutEvaluatedValues)
+	{
+		return false;
+	}
+
 protected:
 	//~ Begin UObject
 	PROPERTYANIMATORCORE_API virtual void PostLoad() override;
@@ -111,7 +129,7 @@ private:
 	void OnGroupNameChanged();
 
 	/** Sets the evaluation result for the resolved property */
-	PROPERTYANIMATORCORE_API void SetEvaluationResult(const FPropertyAnimatorCoreData& InResolvedProperty, const FInstancedPropertyBag& InEvaluatedValues);
+	PROPERTYANIMATORCORE_API void CommitEvaluationResult(const FPropertyAnimatorCoreData& InResolvedProperty, const FInstancedPropertyBag& InEvaluatedValues);
 
 	/** Use this to resolve virtual linked property */
 	PROPERTYANIMATORCORE_API TArray<FPropertyAnimatorCoreData> ResolveProperty(bool bInForEvaluation) const;
@@ -135,6 +153,14 @@ private:
 	UPROPERTY(EditInstanceOnly, Setter="SetAnimated", Getter="IsAnimated", Category="Animator", meta=(AllowPrivateAccess="true"))
 	bool bAnimated = true;
 
+	/** Magnitude of the effect on this property */
+	UPROPERTY(EditInstanceOnly, Setter, Getter, Category="Animator", meta=(ClampMin="0.0", ClampMax="1.0"))
+	float Magnitude = 1.f;
+
+	/** Time offset variation for evaluation */
+	UPROPERTY(EditInstanceOnly, Setter, Getter, Category="Animator")
+	double TimeOffset = 0.f;
+
 	/** Edit condition for modes */
 	UPROPERTY(Transient)
 	bool bEditMode = true;
diff --git a/Engine/Plugins/Experimental/PropertyAnimatorCore/Source/PropertyAnimatorCore/Public/Properties/PropertyAnimatorCoreData.h b/Engine/Plugins/Experimental/PropertyAnimatorCore/Source/PropertyAnimatorCore/Public/Properties/PropertyAnimatorCoreData.h
--- a/Engine/Plugins/Experimental/PropertyAnimatorCore/Source/PropertyAnimatorCore/Public/Properties/PropertyAnimatorCoreData.h
+++ b/Engine/Plugins/Experimental/PropertyAnimatorCore/Source/PropertyAnimatorCore/Public/Properties/PropertyAnimatorCoreData.h
@@ -109,9 +109,13 @@ struct FPropertyAnimatorCoreData
 	/** The member property name */
 	PROPERTYANIMATORCORE_API FName GetMemberPropertyName() const;
 
+	PROPERTYANIMATORCORE_API FName GetMemberPropertyTypeName() const;
+
 	/** The leaf property name */
 	PROPERTYANIMATORCORE_API FName GetLeafPropertyName() const;
 
+	PROPERTYANIMATORCORE_API FName GetLeafPropertyTypeName() const;
+
 	/** The chain properties from member to inner property */
 	PROPERTYANIMATORCORE_API TArray<FProperty*> GetChainProperties() const;
 
@@ -180,6 +184,8 @@ struct FPropertyAnimatorCoreData
 private:
 	static void CopyPropertyValue(const FProperty* InProperty, const void* InSrc, void* OutDest);
 
+	static FName GetPropertyTypeName(const FProperty* InProperty);
+
 	/** Internal use only, used to quickly identify a data within a controller */
 	explicit FPropertyAnimatorCoreData(const FString& InPathHash, FName InDisplayName);
 
diff --git a/Engine/Plugins/Experimental/PropertyAnimatorCore/Source/PropertyAnimatorCore/Public/TimeSources/PropertyAnimatorCoreRangeTimeSource.h b/Engine/Plugins/Experimental/PropertyAnimatorCore/Source/PropertyAnimatorCore/Public/TimeSources/PropertyAnimatorCoreRangeTimeSource.h
--- a/Engine/Plugins/Experimental/PropertyAnimatorCore/Source/PropertyAnimatorCore/Public/TimeSources/PropertyAnimatorCoreRangeTimeSource.h
+++ b/Engine/Plugins/Experimental/PropertyAnimatorCore/Source/PropertyAnimatorCore/Public/TimeSources/PropertyAnimatorCoreRangeTimeSource.h
@@ -55,15 +55,15 @@ protected:
 	PROPERTYANIMATORCORE_API virtual bool IsValidTimeElapsed(double InTimeElapsed) const override;
 	//~ End UPropertyAnimatorTimeSourceBase
 
-	UPROPERTY(EditInstanceOnly, BlueprintReadWrite, Setter="SetUseStartTime", Getter="GetUseStartTime", Category="Animator")
+	UPROPERTY(EditInstanceOnly, BlueprintReadWrite, Setter="SetUseStartTime", Getter="GetUseStartTime", Category="Animator", meta=(InlineEditConditionToggle))
 	bool bUseStartTime = false;
 
-	UPROPERTY(EditInstanceOnly, BlueprintReadWrite, Setter, Getter, Category="Animator", meta=(EditCondition="bUseStartTime", EditConditionHides))
+	UPROPERTY(EditInstanceOnly, BlueprintReadWrite, Setter, Getter, Category="Animator", meta=(EditCondition="bUseStartTime"))
 	double StartTime = 0;
 
-	UPROPERTY(EditInstanceOnly, BlueprintReadWrite, Setter="SetUseStopTime", Getter="GetUseStopTime", Category="Animator")
+	UPROPERTY(EditInstanceOnly, BlueprintReadWrite, Setter="SetUseStopTime", Getter="GetUseStopTime", Category="Animator", meta=(InlineEditConditionToggle))
 	bool bUseStopTime = false;
 
-	UPROPERTY(EditInstanceOnly, BlueprintReadWrite, Setter, Getter, Category="Animator", meta=(EditCondition="bUseStopTime", EditConditionHides))
+	UPROPERTY(EditInstanceOnly, BlueprintReadWrite, Setter, Getter, Category="Animator", meta=(EditCondition="bUseStopTime"))
 	double StopTime = 0;
 };
\ No newline at end of file
diff --git a/Engine/Plugins/Experimental/PropertyAnimatorCore/Source/PropertyAnimatorCore/Public/TimeSources/PropertyAnimatorCoreTimeSourceBase.h b/Engine/Plugins/Experimental/PropertyAnimatorCore/Source/PropertyAnimatorCore/Public/TimeSources/PropertyAnimatorCoreTimeSourceBase.h
--- a/Engine/Plugins/Experimental/PropertyAnimatorCore/Source/PropertyAnimatorCore/Public/TimeSources/PropertyAnimatorCoreTimeSourceBase.h
+++ b/Engine/Plugins/Experimental/PropertyAnimatorCore/Source/PropertyAnimatorCore/Public/TimeSources/PropertyAnimatorCoreTimeSourceBase.h
@@ -18,6 +18,10 @@ class UPropertyAnimatorCoreTimeSourceBase : public UObject
 	friend class UPropertyAnimatorCoreSubsystem;
 
 public:
+#if WITH_EDITOR
+	PROPERTYANIMATORCORE_API static FName GetTimeElapsedPropertyName();
+#endif
+
 	UPropertyAnimatorCoreTimeSourceBase()
 		: UPropertyAnimatorCoreTimeSourceBase(NAME_None)
 	{}
diff --git a/Engine/Plugins/Experimental/PropertyAnimatorCore/Source/PropertyAnimatorCoreEditor/Private/Customizations/PropertyAnimatorCoreEditorContextTypeCustomization.cpp b/Engine/Plugins/Experimental/PropertyAnimatorCore/Source/PropertyAnimatorCoreEditor/Private/Customizations/PropertyAnimatorCoreEditorContextTypeCustomization.cpp
--- a/Engine/Plugins/Experimental/PropertyAnimatorCore/Source/PropertyAnimatorCoreEditor/Private/Customizations/PropertyAnimatorCoreEditorContextTypeCustomization.cpp
+++ b/Engine/Plugins/Experimental/PropertyAnimatorCore/Source/PropertyAnimatorCoreEditor/Private/Customizations/PropertyAnimatorCoreEditorContextTypeCustomization.cpp
@@ -35,6 +35,7 @@ void FPropertyAnimatorCoreEditorContextTypeCustomization::CustomizeHeader(TShare
 	PropertyContextHandle = InPropertyHandle;
 
 	const FName PropertyName = PropertyContext->GetAnimatedProperty().GetPropertyDisplayName();
+	const FName PropertyTypeName = PropertyContext->GetAnimatedProperty().GetLeafPropertyTypeName();
 
 	InRow
 		.NameContent()
@@ -50,7 +51,7 @@ void FPropertyAnimatorCoreEditorContextTypeCustomization::CustomizeHeader(TShare
 			+ SHorizontalBox::Slot()
 			.FillWidth(1.f)
 			[
-				InPropertyHandle->CreatePropertyNameWidget(FText::FromName(PropertyName))
+				InPropertyHandle->CreatePropertyNameWidget(FText::FromString(PropertyName.ToString() + TEXT(" (") + PropertyTypeName.ToString() + TEXT(")")))
 			]
 		]
 		.ValueContent()
diff --git a/Engine/Plugins/Experimental/PropertyAnimatorCore/Source/PropertyAnimatorCoreEditor/Private/Menus/PropertyAnimatorCoreEditorMenu.cpp b/Engine/Plugins/Experimental/PropertyAnimatorCore/Source/PropertyAnimatorCoreEditor/Private/Menus/PropertyAnimatorCoreEditorMenu.cpp
--- a/Engine/Plugins/Experimental/PropertyAnimatorCore/Source/PropertyAnimatorCoreEditor/Private/Menus/PropertyAnimatorCoreEditorMenu.cpp
+++ b/Engine/Plugins/Experimental/PropertyAnimatorCore/Source/PropertyAnimatorCoreEditor/Private/Menus/PropertyAnimatorCoreEditorMenu.cpp
@@ -506,7 +506,7 @@ void UE::PropertyAnimatorCoreEditor::Menu::ExecuteLinkAnimatorPropertyAction(UPr
 		|| !IsValid(InAnimator)
 		|| InAnimator->IsTemplate()
 		|| !InProperty.IsResolved()
-		|| !InAnimator->IsPropertySupported(InProperty)
+		|| !InAnimator->HasPropertySupport(InProperty)
 		|| InMenuData->GetContext().IsEmpty())
 	{
 		return;
@@ -733,13 +733,13 @@ void UE::PropertyAnimatorCoreEditor::Menu::FillNewAnimatorSubmenu(UToolMenu* InM
 		TSet<FPropertyAnimatorCoreData> SupportedProperties;
 		for (const FPropertyAnimatorCoreData& Property : InMenuData->GetContext().GetProperties())
 		{
-			InAnimator->GetPropertiesSupported(Property, SupportedProperties, true);
+			InAnimator->GetPropertiesSupported(Property, SupportedProperties, /** SearchDepth */3);
 		}
 
-		for (FPropertyAnimatorCoreData& SupportedProperty : SupportedProperties)
+		for (const FPropertyAnimatorCoreData& SupportedProperty : SupportedProperties)
 		{
 			const FString MenuName = SupportedProperty.GetPropertyDisplayName().ToString();
-			const FText MenuLabel = FText::FromString(MenuName);
+			const FText MenuLabel = FText::FromString(MenuName + TEXT(" (") + SupportedProperty.GetLeafPropertyTypeName().ToString() + TEXT(")"));
 
 			// Create action (creates an animator and links the property)
 			PropertySection.AddMenuEntry(
@@ -828,13 +828,13 @@ void UE::PropertyAnimatorCoreEditor::Menu::FillLinkAnimatorSubmenu(UToolMenu* In
 	TSet<FPropertyAnimatorCoreData> SupportedProperties;
 	for (const FPropertyAnimatorCoreData& Property : InMenuData->GetContext().GetProperties())
 	{
-		InAnimator->GetPropertiesSupported(Property, SupportedProperties, true);
+		InAnimator->GetPropertiesSupported(Property, SupportedProperties, /** SearchDepth */3);
 	}
 
 	for (const FPropertyAnimatorCoreData& SupportedProperty : SupportedProperties)
 	{
 		const FName MenuName = SupportedProperty.GetPropertyDisplayName();
-		const FText MenuLabel = FText::FromName(MenuName);
+		const FText MenuLabel = FText::FromString(MenuName.ToString() + TEXT(" (") + SupportedProperty.GetLeafPropertyTypeName().ToString() + TEXT(")"));
 
 		PropertySection.AddMenuEntry(
 			MenuName
@@ -886,7 +886,7 @@ void UE::PropertyAnimatorCoreEditor::Menu::FillPresetAnimatorSubmenu(UToolMenu*
 	for (const FPropertyAnimatorCoreData& SupportedProperty : SupportedProperties)
 	{
 		const FName MenuName = SupportedProperty.GetPropertyDisplayName();
-		const FText MenuLabel = FText::FromName(MenuName);
+		const FText MenuLabel = FText::FromString(MenuName.ToString() + TEXT(" (") + SupportedProperty.GetLeafPropertyTypeName().ToString() + TEXT(")"));
 		const FText MenuTooltip = LOCTEXT("LinkPresetProperty.Tooltip", "Link this preset property");
 
 		const FToolMenuEntry SupportedPropertyEntry = FToolMenuEntry::InitMenuEntry(
@@ -976,7 +976,7 @@ void UE::PropertyAnimatorCoreEditor::Menu::FillNewPresetAnimatorSubmenu(UToolMen
 
 	for (const FPropertyAnimatorCoreData& PresetProperty : PresetProperties)
 	{
-		const FText MenuLabel = FText::FromName(PresetProperty.GetPropertyDisplayName());
+		const FText MenuLabel = FText::FromString(PresetProperty.GetPropertyDisplayName().ToString() + TEXT(" (") + PresetProperty.GetLeafPropertyTypeName().ToString() + TEXT(")"));
 
 		const FToolMenuEntry CreatePropertyEntry = FToolMenuEntry::InitMenuEntry(
 			FName(TEXT("CreateProperty") + PresetProperty.GetPropertyDisplayName().ToString())
diff --git a/Engine/Plugins/Experimental/PropertyAnimatorCore/Source/PropertyAnimatorCoreEditor/Private/Widgets/TableRows/SPropertyAnimatorCoreEditorControllersViewTableRow.cpp b/Engine/Plugins/Experimental/PropertyAnimatorCore/Source/PropertyAnimatorCoreEditor/Private/Widgets/TableRows/SPropertyAnimatorCoreEditorControllersViewTableRow.cpp
--- a/Engine/Plugins/Experimental/PropertyAnimatorCore/Source/PropertyAnimatorCoreEditor/Private/Widgets/TableRows/SPropertyAnimatorCoreEditorControllersViewTableRow.cpp
+++ b/Engine/Plugins/Experimental/PropertyAnimatorCore/Source/PropertyAnimatorCoreEditor/Private/Widgets/TableRows/SPropertyAnimatorCoreEditorControllersViewTableRow.cpp
@@ -224,18 +224,18 @@ TOptional<EItemDropZone> SPropertyAnimatorCoreEditorControllersViewTableRow::OnP
 	// Get dragged properties and check that the controller we drop them on supports them
 	if (const TSharedPtr<FPropertyAnimatorCoreEditorViewDragDropOp> DragDropOp = InDragDropEvent.GetOperationAs<FPropertyAnimatorCoreEditorViewDragDropOp>())
 	{
-		const UPropertyAnimatorCoreBase* Controller = InItem->ControlledProperty.ControllerWeak.Get();
+		const UPropertyAnimatorCoreBase* Animator = InItem->ControlledProperty.ControllerWeak.Get();
 
 		if (InItem.IsValid()
 			&& !InItem->ControlledProperty.Property.IsValid()
-			&& Controller)
+			&& Animator)
 		{
 			for (const FPropertiesViewControllerItem& DraggedItem : DragDropOp->GetDraggedItems())
 			{
 				if (const FPropertyAnimatorCoreData* Property = DraggedItem.Property.Get())
 				{
-					if (Controller->IsPropertySupported(*Property)
-						&& !Controller->IsPropertyLinked(*Property))
+					if (Animator->HasPropertySupport(*Property)
+						&& !Animator->IsPropertyLinked(*Property))
 					{
 						return EItemDropZone::OntoItem;
 					}
diff --git a/Engine/Plugins/Experimental/PropertyAnimatorCore/Source/PropertyAnimatorCoreEditor/Private/Widgets/TableRows/SPropertyAnimatorCoreEditorPropertiesViewTableRow.cpp b/Engine/Plugins/Experimental/PropertyAnimatorCore/Source/PropertyAnimatorCoreEditor/Private/Widgets/TableRows/SPropertyAnimatorCoreEditorPropertiesViewTableRow.cpp
--- a/Engine/Plugins/Experimental/PropertyAnimatorCore/Source/PropertyAnimatorCoreEditor/Private/Widgets/TableRows/SPropertyAnimatorCoreEditorPropertiesViewTableRow.cpp
+++ b/Engine/Plugins/Experimental/PropertyAnimatorCore/Source/PropertyAnimatorCoreEditor/Private/Widgets/TableRows/SPropertyAnimatorCoreEditorPropertiesViewTableRow.cpp
@@ -158,10 +158,10 @@ TOptional<EItemDropZone> SPropertyAnimatorCoreEditorPropertiesViewTableRow::OnPr
 		{
 			for (const FPropertiesViewControllerItem& DraggedItem : DragDropOp->GetDraggedItems())
 			{
-				if (const UPropertyAnimatorCoreBase* Controller = DraggedItem.ControllerWeak.Get())
+				if (const UPropertyAnimatorCoreBase* Animator = DraggedItem.ControllerWeak.Get())
 				{
-					if (Controller->IsPropertySupported(InItem->Property)
-						&& !Controller->IsPropertyLinked(InItem->Property))
+					if (Animator->HasPropertySupport(InItem->Property)
+						&& !Animator->IsPropertyLinked(InItem->Property))
 					{
 						return EItemDropZone::OntoItem;
 					}
@@ -202,16 +202,16 @@ FReply SPropertyAnimatorCoreEditorPropertiesViewTableRow::OnPropertyDrop(FDragDr
 
 			for (const FPropertiesViewControllerItem& DraggedItem : DragDropOp->GetDraggedItems())
 			{
-				if (UPropertyAnimatorCoreBase* Controller = DraggedItem.ControllerWeak.Get())
+				if (UPropertyAnimatorCoreBase* Animator = DraggedItem.ControllerWeak.Get())
 				{
-					if (Controller->IsPropertySupported(RowItem->Property))
+					if (Animator->HasPropertySupport(RowItem->Property))
 					{
 						if (!bIsControlPressed && DraggedItem.Property.IsValid())
 						{
-							Controller->UnlinkProperty(*DraggedItem.Property);
+							Animator->UnlinkProperty(*DraggedItem.Property);
 						}
 
-						Controller->LinkProperty(RowItem->Property);
+						Animator->LinkProperty(RowItem->Property);
 					}
 				}
 			}
diff --git a/Engine/Plugins/Experimental/PropertyAnimatorCore/Source/PropertyAnimatorCoreEditor/Private/Widgets/Views/SPropertyAnimatorCoreEditorPropertiesView.cpp b/Engine/Plugins/Experimental/PropertyAnimatorCore/Source/PropertyAnimatorCoreEditor/Private/Widgets/Views/SPropertyAnimatorCoreEditorPropertiesView.cpp
--- a/Engine/Plugins/Experimental/PropertyAnimatorCore/Source/PropertyAnimatorCoreEditor/Private/Widgets/Views/SPropertyAnimatorCoreEditorPropertiesView.cpp
+++ b/Engine/Plugins/Experimental/PropertyAnimatorCore/Source/PropertyAnimatorCoreEditor/Private/Widgets/Views/SPropertyAnimatorCoreEditorPropertiesView.cpp
@@ -170,12 +170,10 @@ void SPropertyAnimatorCoreEditorPropertiesView::OnGetChildren(FPropertiesViewIte
 			TSet<UPropertyAnimatorCoreBase*> ExistingControllers = ControllerSubsystem->GetExistingAnimators(InItem->Property);
 			TSet<UPropertyAnimatorCoreBase*> AvailableControllers = ControllerSubsystem->GetAvailableAnimators(&InItem->Property);
 
-			constexpr bool bRecursiveSearch = true;
-
 			for (UPropertyAnimatorCoreBase* Controller : ExistingControllers)
 			{
 				TSet<FPropertyAnimatorCoreData> SupportedProperties;
-				Controller->GetPropertiesSupported(InItem->Property, SupportedProperties, bRecursiveSearch);
+				Controller->GetPropertiesSupported(InItem->Property, SupportedProperties, /** SearchDepth */3);
 
 				// Add available properties to control with this controller
 				for (FPropertyAnimatorCoreData& SupportedProperty : SupportedProperties)
@@ -197,7 +195,7 @@ void SPropertyAnimatorCoreEditorPropertiesView::OnGetChildren(FPropertiesViewIte
 			for (UPropertyAnimatorCoreBase* Controller : AvailableControllers)
 			{
 				TSet<FPropertyAnimatorCoreData> SupportedProperties;
-				Controller->GetPropertiesSupported(InItem->Property, SupportedProperties, bRecursiveSearch);
+				Controller->GetPropertiesSupported(InItem->Property, SupportedProperties, /** SearchDepth */3);
 
 				// Add all supported properties from new controller
 				for (FPropertyAnimatorCoreData& SupportedProperty : SupportedProperties)
