MotionDesign : PropertyAnimator
- AnimatorComponent properties will be shown in the body instead of footer of the operator stack

diff --git a/Engine/Plugins/Experimental/PropertyAnimatorCore/Source/PropertyAnimatorCoreEditor/Private/Customizations/PropertyAnimatorCoreEditorStackCustomization.cpp b/Engine/Plugins/Experimental/PropertyAnimatorCore/Source/PropertyAnimatorCoreEditor/Private/Customizations/PropertyAnimatorCoreEditorStackCustomization.cpp
--- a/Engine/Plugins/Experimental/PropertyAnimatorCore/Source/PropertyAnimatorCoreEditor/Private/Customizations/PropertyAnimatorCoreEditorStackCustomization.cpp
+++ b/Engine/Plugins/Experimental/PropertyAnimatorCore/Source/PropertyAnimatorCoreEditor/Private/Customizations/PropertyAnimatorCoreEditorStackCustomization.cpp
@@ -12,6 +12,7 @@
 #include "Styling/SlateIconFinder.h"
 #include "Subsystems/PropertyAnimatorCoreEditorSubsystem.h"
 #include "Subsystems/PropertyAnimatorCoreSubsystem.h"
+#include "TimeSources/PropertyAnimatorCoreTimeSourceBase.h"
 #include "ToolMenus.h"
 #include "ToolMenu.h"
 #include "Widgets/PropertyAnimatorCoreEditorEditPanelOptions.h"
@@ -282,34 +283,29 @@ void UPropertyAnimatorCoreEditorStackCustomization::CustomizeItemHeader(const FO
 
 void UPropertyAnimatorCoreEditorStackCustomization::CustomizeItemBody(const FOperatorStackEditorItemPtr& InItem, const FOperatorStackEditorTree& InItemTree, FOperatorStackEditorBodyBuilder& InBodyBuilder)
 {
+	if (InItem->IsA<UPropertyAnimatorCoreComponent>())
+	{
+		FProperty* MagnitudeProperty = FindFProperty<FProperty>(UPropertyAnimatorCoreComponent::StaticClass(), GET_MEMBER_NAME_CHECKED(UPropertyAnimatorCoreComponent, AnimatorsMagnitude));
+
+		InBodyBuilder
+			.AllowProperty(MagnitudeProperty)
+			.SetShowDetailsView(true);
+	}
 	// Customize animator body
-	if (InItem->IsA<UPropertyAnimatorCoreBase>())
+	else if (InItem->IsA<UPropertyAnimatorCoreBase>())
 	{
 		FBoolProperty* EnableProperty = FindFProperty<FBoolProperty>(UPropertyAnimatorCoreBase::StaticClass(), GET_MEMBER_NAME_CHECKED(UPropertyAnimatorCoreBase, bAnimatorEnabled));
+		FProperty* TimeElapsedProperty = FindFProperty<FProperty>(UPropertyAnimatorCoreTimeSourceBase::StaticClass(), UPropertyAnimatorCoreTimeSourceBase::GetTimeElapsedPropertyName());
 
 		InBodyBuilder
 			.SetShowDetailsView(true)
-			.DisallowProperty(EnableProperty);
+			.DisallowProperty(EnableProperty)
+			.DisallowProperty(TimeElapsedProperty);
 	}
 
 	Super::CustomizeItemBody(InItem, InItemTree, InBodyBuilder);
 }
 
-void UPropertyAnimatorCoreEditorStackCustomization::CustomizeItemFooter(const FOperatorStackEditorItemPtr& InItem, const FOperatorStackEditorTree& InItemTree, FOperatorStackEditorFooterBuilder& InFooterBuilder)
-{
-	// Customize component footer
-	if (InItem->IsA<UPropertyAnimatorCoreComponent>())
-	{
-		FProperty* MagnitudeProperty = FindFProperty<FProperty>(UPropertyAnimatorCoreComponent::StaticClass(), GET_MEMBER_NAME_CHECKED(UPropertyAnimatorCoreComponent, AnimatorsMagnitude));
-
-		InFooterBuilder
-			.AllowProperty(MagnitudeProperty)
-			.SetShowDetailsView(true);
-	}
-
-	Super::CustomizeItemFooter(InItem, InItemTree, InFooterBuilder);
-}
-
 bool UPropertyAnimatorCoreEditorStackCustomization::OnIsItemSelectable(const FOperatorStackEditorItemPtr& InItem)
 {
 	if (InItem->IsA<UPropertyAnimatorCoreBase>())
diff --git a/Engine/Plugins/Experimental/PropertyAnimatorCore/Source/PropertyAnimatorCoreEditor/Private/Customizations/PropertyAnimatorCoreEditorStackCustomization.h b/Engine/Plugins/Experimental/PropertyAnimatorCore/Source/PropertyAnimatorCoreEditor/Private/Customizations/PropertyAnimatorCoreEditorStackCustomization.h
--- a/Engine/Plugins/Experimental/PropertyAnimatorCore/Source/PropertyAnimatorCoreEditor/Private/Customizations/PropertyAnimatorCoreEditorStackCustomization.h
+++ b/Engine/Plugins/Experimental/PropertyAnimatorCore/Source/PropertyAnimatorCoreEditor/Private/Customizations/PropertyAnimatorCoreEditorStackCustomization.h
@@ -25,7 +25,6 @@ public:
 	virtual void CustomizeStackHeader(const FOperatorStackEditorTree& InItemTree, FOperatorStackEditorHeaderBuilder& InHeaderBuilder) override;
 	virtual void CustomizeItemHeader(const FOperatorStackEditorItemPtr& InItem, const FOperatorStackEditorTree& InItemTree, FOperatorStackEditorHeaderBuilder& InHeaderBuilder) override;
 	virtual void CustomizeItemBody(const FOperatorStackEditorItemPtr& InItem, const FOperatorStackEditorTree& InItemTree, FOperatorStackEditorBodyBuilder& InBodyBuilder) override;
-	virtual void CustomizeItemFooter(const FOperatorStackEditorItemPtr& InItem, const FOperatorStackEditorTree& InItemTree, FOperatorStackEditorFooterBuilder& InFooterBuilder) override;
 	virtual bool OnIsItemSelectable(const FOperatorStackEditorItemPtr& InItem) override;
 	virtual const FSlateBrush* GetIcon() const override;
 	//~ End UOperatorStackEditorStackCustomization
