MotionDesign : PropertyAnimator
- Adding machine clock time source to use local time, utc time, countdown, stopwatch

diff --git a/Engine/Plugins/Experimental/PropertyAnimatorCore/Source/PropertyAnimatorCore/Private/TimeSources/PropertyAnimatorCoreMachineClockTimeSource.cpp b/Engine/Plugins/Experimental/PropertyAnimatorCore/Source/PropertyAnimatorCore/Private/TimeSources/PropertyAnimatorCoreMachineClockTimeSource.cpp
new file mode 100644
--- /dev/null
+++ b/Engine/Plugins/Experimental/PropertyAnimatorCore/Source/PropertyAnimatorCore/Private/TimeSources/PropertyAnimatorCoreMachineClockTimeSource.cpp
@@ -0,0 +1,139 @@
+// Copyright Epic Games, Inc. All Rights Reserved.
+
+#include "TimeSources/PropertyAnimatorCoreMachineClockTimeSource.h"
+
+#include "Internationalization/Regex.h"
+#include "Misc/DateTime.h"
+
+double UPropertyAnimatorCoreMachineClockTimeSource::GetTimeElapsed()
+{
+	double TimeElapsedSeconds = 0;
+
+	switch(Mode)
+	{
+	case EPropertyAnimatorCoreMachineClockMode::LocalTime:
+		{
+			TimeElapsedSeconds = (FDateTime::Now() - FDateTime::MinValue()).GetTotalSeconds();
+		}
+		break;
+	case EPropertyAnimatorCoreMachineClockMode::UtcTime:
+		{
+			TimeElapsedSeconds = (FDateTime::UtcNow() - FDateTime::MinValue()).GetTotalSeconds();
+		}
+		break;
+	case EPropertyAnimatorCoreMachineClockMode::Countdown:
+		{
+			TimeElapsedSeconds = (CountdownTimeSpan - (FDateTime::Now() - ActivationTime)).GetTotalSeconds();
+		}
+		break;
+	case EPropertyAnimatorCoreMachineClockMode::Stopwatch:
+		{
+			TimeElapsedSeconds = (FDateTime::Now() - ActivationTime).GetTotalSeconds();
+		}
+		break;
+	}
+
+	return TimeElapsedSeconds;
+}
+
+bool UPropertyAnimatorCoreMachineClockTimeSource::IsTimeSourceReady() const
+{
+	return true;
+}
+
+void UPropertyAnimatorCoreMachineClockTimeSource::OnTimeSourceActive()
+{
+	Super::OnTimeSourceActive();
+
+	ActivationTime = FDateTime::Now();
+	OnModeChanged();
+}
+
+void UPropertyAnimatorCoreMachineClockTimeSource::SetMode(EPropertyAnimatorCoreMachineClockMode InMode)
+{
+	if (Mode == InMode)
+	{
+		return;
+	}
+
+	Mode = InMode;
+	OnModeChanged();
+}
+
+void UPropertyAnimatorCoreMachineClockTimeSource::SetCountdownDuration(const FTimespan& InTimeSpan)
+{
+	if (InTimeSpan == CountdownTimeSpan)
+	{
+		return;
+	}
+
+	SetCountdownDuration(InTimeSpan.ToString(TEXT("%h:%m:%s")));
+}
+
+void UPropertyAnimatorCoreMachineClockTimeSource::SetCountdownDuration(const FString& InDuration)
+{
+	if (CountdownDuration == InDuration)
+	{
+		return;
+	}
+
+	CountdownDuration = InDuration;
+	OnModeChanged();
+}
+
+#if WITH_EDITOR
+void UPropertyAnimatorCoreMachineClockTimeSource::PostEditChangeProperty(FPropertyChangedEvent& InPropertyChangedEvent)
+{
+	Super::PostEditChangeProperty(InPropertyChangedEvent);
+
+	const FName MemberName = InPropertyChangedEvent.GetMemberPropertyName();
+
+	if (MemberName == GET_MEMBER_NAME_CHECKED(UPropertyAnimatorCoreMachineClockTimeSource, Mode)
+		|| MemberName == GET_MEMBER_NAME_CHECKED(UPropertyAnimatorCoreMachineClockTimeSource, CountdownDuration))
+	{
+		OnModeChanged();
+	}
+}
+#endif
+
+void UPropertyAnimatorCoreMachineClockTimeSource::OnModeChanged()
+{
+	if (Mode == EPropertyAnimatorCoreMachineClockMode::Countdown)
+	{
+		CountdownTimeSpan = ParseTime(CountdownDuration);
+	}
+}
+
+FTimespan UPropertyAnimatorCoreMachineClockTimeSource::ParseTime(const FString& InFormat)
+{
+	// Regex patterns for different formats
+	static const FRegexPattern HHMMSSPattern(TEXT("^(?:(\\d{2}):)?(\\d{2}):(\\d{2})$")); // 01:00 00:01:00
+	static const FRegexPattern CombinedPattern(TEXT("(?:(\\d+)h)? ?(?:(\\d+)m)? ?(?:(\\d+)s)?")); // 1h 1m 1s
+
+	FRegexMatcher HHMMSSMatcher(HHMMSSPattern, InFormat);
+	FRegexMatcher CombinedMatcher(CombinedPattern, InFormat);
+
+	FTimespan ParsedTimeSpan = FTimespan::Zero();
+
+	if (InFormat.IsNumeric())
+	{
+		const int32 Seconds = FCString::Atoi(*InFormat);
+		ParsedTimeSpan = FTimespan::FromSeconds(Seconds);
+	}
+	else if (HHMMSSMatcher.FindNext())
+	{
+		const int32 Hours = HHMMSSMatcher.GetCaptureGroup(1).IsEmpty() ? 0 : FCString::Atoi(*HHMMSSMatcher.GetCaptureGroup(1));
+		const int32 Minutes = FCString::Atoi(*HHMMSSMatcher.GetCaptureGroup(2));
+		const int32 Seconds = FCString::Atoi(*HHMMSSMatcher.GetCaptureGroup(3));
+		ParsedTimeSpan = FTimespan::FromHours(Hours) + FTimespan::FromMinutes(Minutes) + FTimespan::FromSeconds(Seconds);
+	}
+	else if (CombinedMatcher.FindNext())
+	{
+		const int32 Hours = CombinedMatcher.GetCaptureGroup(1).IsEmpty() ? 0 : FCString::Atoi(*CombinedMatcher.GetCaptureGroup(1));
+		const int32 Minutes = CombinedMatcher.GetCaptureGroup(2).IsEmpty() ? 0 : FCString::Atoi(*CombinedMatcher.GetCaptureGroup(2));
+		const int32 Seconds = CombinedMatcher.GetCaptureGroup(3).IsEmpty() ? 0 : FCString::Atoi(*CombinedMatcher.GetCaptureGroup(3));
+		ParsedTimeSpan = FTimespan::FromHours(Hours) + FTimespan::FromMinutes(Minutes) + FTimespan::FromSeconds(Seconds);
+	}
+
+	return ParsedTimeSpan;
+}
diff --git a/Engine/Plugins/Experimental/PropertyAnimatorCore/Source/PropertyAnimatorCore/Public/TimeSources/PropertyAnimatorCoreMachineClockTimeSource.h b/Engine/Plugins/Experimental/PropertyAnimatorCore/Source/PropertyAnimatorCore/Public/TimeSources/PropertyAnimatorCoreMachineClockTimeSource.h
new file mode 100644
--- /dev/null
+++ b/Engine/Plugins/Experimental/PropertyAnimatorCore/Source/PropertyAnimatorCore/Public/TimeSources/PropertyAnimatorCoreMachineClockTimeSource.h
@@ -0,0 +1,92 @@
+// Copyright Epic Games, Inc. All Rights Reserved.
+
+#pragma once
+
+#include "Misc/DateTime.h"
+#include "Misc/Timespan.h"
+#include "PropertyAnimatorCoreTimeSourceBase.h"
+#include "PropertyAnimatorCoreMachineClockTimeSource.generated.h"
+
+/** Enumerates all possible modes for the machine clock time source */
+UENUM(BlueprintType)
+enum class EPropertyAnimatorCoreMachineClockMode : uint8
+{
+	/** Local time of the machine */
+	LocalTime,
+	/** Universal time = Greenwich Mean Time */
+	UtcTime,
+	/** Specified duration elapsing until it reaches 0 */
+	Countdown,
+	/** Current time elapsed since the time source is active */
+	Stopwatch
+};
+
+/** Machine clock time source that support various option */
+UCLASS(MinimalAPI)
+class UPropertyAnimatorCoreMachineClockTimeSource : public UPropertyAnimatorCoreTimeSourceBase
+{
+	GENERATED_BODY()
+
+public:
+	UPropertyAnimatorCoreMachineClockTimeSource()
+		: UPropertyAnimatorCoreTimeSourceBase(TEXT("MachineClock"))
+	{}
+
+	//~ Begin UPropertyAnimatorTimeSourceBase
+	virtual double GetTimeElapsed() override;
+	virtual bool IsTimeSourceReady() const override;
+	virtual void OnTimeSourceActive() override;
+	//~ End UPropertyAnimatorTimeSourceBase
+
+	void SetMode(EPropertyAnimatorCoreMachineClockMode InMode);
+	EPropertyAnimatorCoreMachineClockMode GetMode() const
+	{
+		return Mode;
+	}
+
+	void SetCountdownDuration(const FTimespan& InTimeSpan);
+	void GetCountdownDuration(FTimespan& OutTimeSpan) const
+	{
+		OutTimeSpan = CountdownTimeSpan;
+	}
+
+	void SetCountdownDuration(const FString& InDuration);
+	const FString& GetCountdownDuration() const
+	{
+		return CountdownDuration;
+	}
+
+protected:
+	static FTimespan ParseTime(const FString& InFormat);
+
+	//~ Begin UObject
+#if WITH_EDITOR
+	virtual void PostEditChangeProperty(FPropertyChangedEvent& InPropertyChangedEvent) override;
+#endif
+	//~ End UObject
+
+	void OnModeChanged();
+
+	/** Machine time mode to use */
+	UPROPERTY(EditInstanceOnly, Setter, Getter, Category="Animator")
+	EPropertyAnimatorCoreMachineClockMode Mode = EPropertyAnimatorCoreMachineClockMode::LocalTime;
+
+	/**
+	* Countdown duration format : 
+	* 120 = 2 minutes
+	* 02:00 = 2 minutes
+	* 00:02:00 = 2 minutes
+	* 2m = 2 minutes
+	* 1h = 1 hour
+	* 120s = 2 minutes
+	*/
+	UPROPERTY(EditInstanceOnly, Setter, Getter, Category="Animator", meta=(EditCondition="Mode == EPropertyAnimatorCoreLocalTimeMode::Countdown", EditConditionHides))
+	FString CountdownDuration = TEXT("1m");
+
+private:
+	UPROPERTY(Transient)
+	FTimespan CountdownTimeSpan;
+
+	UPROPERTY(Transient)
+	FDateTime ActivationTime;
+};
\ No newline at end of file
