MotionDesign : PropertyAnimator
- Removing "Edit Animators" widgets and all menus linked to it

diff --git a/Engine/Plugins/Experimental/Avalanche/Source/AvalanchePropertyAnimatorEditor/Private/Outliner/AvaPropertyAnimatorEditorOutlinerContextMenu.cpp b/Engine/Plugins/Experimental/Avalanche/Source/AvalanchePropertyAnimatorEditor/Private/Outliner/AvaPropertyAnimatorEditorOutlinerContextMenu.cpp
--- a/Engine/Plugins/Experimental/Avalanche/Source/AvalanchePropertyAnimatorEditor/Private/Outliner/AvaPropertyAnimatorEditorOutlinerContextMenu.cpp
+++ b/Engine/Plugins/Experimental/Avalanche/Source/AvalanchePropertyAnimatorEditor/Private/Outliner/AvaPropertyAnimatorEditorOutlinerContextMenu.cpp
@@ -38,8 +38,7 @@ void FAvaPropertyAnimatorEditorOutlinerContextMenu::OnExtendOutlinerContextMenu(
 	const FPropertyAnimatorCoreEditorMenuContext MenuContext(ContextObjects, {});
 	FPropertyAnimatorCoreEditorMenuOptions MenuOptions(
 		{
-			EPropertyAnimatorCoreEditorMenuType::Edit
-			, EPropertyAnimatorCoreEditorMenuType::NewAdvanced
+			EPropertyAnimatorCoreEditorMenuType::NewAdvanced
 			, EPropertyAnimatorCoreEditorMenuType::Enable
 			, EPropertyAnimatorCoreEditorMenuType::Disable
 			, EPropertyAnimatorCoreEditorMenuType::Delete
diff --git a/Engine/Plugins/Experimental/PropertyAnimatorCore/Source/PropertyAnimatorCoreEditor/Private/Customizations/PropertyAnimatorCoreEditorStackCustomization.cpp b/Engine/Plugins/Experimental/PropertyAnimatorCore/Source/PropertyAnimatorCoreEditor/Private/Customizations/PropertyAnimatorCoreEditorStackCustomization.cpp
--- a/Engine/Plugins/Experimental/PropertyAnimatorCore/Source/PropertyAnimatorCoreEditor/Private/Customizations/PropertyAnimatorCoreEditorStackCustomization.cpp
+++ b/Engine/Plugins/Experimental/PropertyAnimatorCore/Source/PropertyAnimatorCoreEditor/Private/Customizations/PropertyAnimatorCoreEditorStackCustomization.cpp
@@ -15,7 +15,6 @@
 #include "TimeSources/PropertyAnimatorCoreTimeSourceBase.h"
 #include "ToolMenus.h"
 #include "ToolMenu.h"
-#include "Widgets/PropertyAnimatorCoreEditorEditPanelOptions.h"
 
 #define LOCTEXT_NAMESPACE "PropertyAnimatorEditorStackCustomization"
 
diff --git a/Engine/Plugins/Experimental/PropertyAnimatorCore/Source/PropertyAnimatorCoreEditor/Private/Menus/PropertyAnimatorCoreEditorMenu.cpp b/Engine/Plugins/Experimental/PropertyAnimatorCore/Source/PropertyAnimatorCoreEditor/Private/Menus/PropertyAnimatorCoreEditorMenu.cpp
--- a/Engine/Plugins/Experimental/PropertyAnimatorCore/Source/PropertyAnimatorCoreEditor/Private/Menus/PropertyAnimatorCoreEditorMenu.cpp
+++ b/Engine/Plugins/Experimental/PropertyAnimatorCore/Source/PropertyAnimatorCoreEditor/Private/Menus/PropertyAnimatorCoreEditorMenu.cpp
@@ -11,49 +11,9 @@
 #include "Styling/SlateIconFinder.h"
 #include "Subsystems/PropertyAnimatorCoreEditorSubsystem.h"
 #include "Subsystems/PropertyAnimatorCoreSubsystem.h"
-#include "Widgets/PropertyAnimatorCoreEditorEditPanelOptions.h"
 
 #define LOCTEXT_NAMESPACE "PropertyAnimatorCoreEditorMenu"
 
-void UE::PropertyAnimatorCoreEditor::Menu::FillEditAnimatorSection(UToolMenu* InMenu, TSharedRef<FPropertyAnimatorCoreEditorMenuData> InMenuData)
-{
-	if (!InMenu || InMenuData->GetContext().IsEmpty())
-	{
-		return;
-	}
-
-	const UPropertyAnimatorCoreSubsystem* Subsystem = UPropertyAnimatorCoreSubsystem::Get();
-	if (!Subsystem)
-	{
-		return;
-	}
-
-	const TSet<AActor*> Actors = InMenuData->GetContext().GetActors();
-	if (Actors.Num() != 1)
-	{
-		return;
-	}
-
-	AActor* AnimatorActor = Actors.Array()[0];
-	if (!IsValid(AnimatorActor))
-	{
-		return;
-	}
-
-	FToolMenuSection& ManageAnimatorsSection = InMenu->FindOrAddSection(TEXT("ManageAnimators"), LOCTEXT("ManageAnimators.Label", "Manage Animators"));
-	const FSlateIcon AnimatorIcon = FSlateIconFinder::FindIconForClass(UPropertyAnimatorCoreBase::StaticClass());
-
-	ManageAnimatorsSection.AddMenuEntry(
-		TEXT("OpenEditAnimatorsWindow")
-	   , LOCTEXT("OpenEditAnimatorsWindow.Label", "Edit Animators")
-	   , LOCTEXT("OpenEditAnimatorsWindow.Tooltip", "Open edit animators window")
-	   , AnimatorIcon
-	   , FUIAction(
-		   FExecuteAction::CreateStatic(&ExecuteEditAnimatorAction, AnimatorActor)
-	   )
-	);
-}
-
 void UE::PropertyAnimatorCoreEditor::Menu::FillNewAnimatorSection(UToolMenu* InMenu, TSharedRef<FPropertyAnimatorCoreEditorMenuData> InMenuData)
 {
 	if (!InMenu || InMenuData->GetContext().IsEmpty())
@@ -370,19 +330,6 @@ void UE::PropertyAnimatorCoreEditor::Menu::FillDisableAnimatorSection(UToolMenu*
 	}
 }
 
-void UE::PropertyAnimatorCoreEditor::Menu::ExecuteEditAnimatorAction(AActor* InActor)
-{
-	UPropertyAnimatorCoreEditorSubsystem* EditorSubsystem = UPropertyAnimatorCoreEditorSubsystem::Get();
-
-	if (!EditorSubsystem || !IsValid(InActor))
-	{
-		return;
-	}
-
-	FPropertyAnimatorCoreEditorEditPanelOptions& Options = EditorSubsystem->OpenPropertyControlWindow();
-	Options.SetContextActor(InActor);
-}
-
 void UE::PropertyAnimatorCoreEditor::Menu::ExecuteNewAnimatorPresetAction(const UPropertyAnimatorCoreBase* InAnimator, const TSet<AActor*>& InActors, UPropertyAnimatorCorePresetBase* InPreset, TSharedRef<FPropertyAnimatorCoreEditorMenuData> InMenuData)
 {
 	const UPropertyAnimatorCoreSubsystem* Subsystem = UPropertyAnimatorCoreSubsystem::Get();
diff --git a/Engine/Plugins/Experimental/PropertyAnimatorCore/Source/PropertyAnimatorCoreEditor/Private/Menus/PropertyAnimatorCoreEditorMenu.h b/Engine/Plugins/Experimental/PropertyAnimatorCore/Source/PropertyAnimatorCoreEditor/Private/Menus/PropertyAnimatorCoreEditorMenu.h
--- a/Engine/Plugins/Experimental/PropertyAnimatorCore/Source/PropertyAnimatorCoreEditor/Private/Menus/PropertyAnimatorCoreEditorMenu.h
+++ b/Engine/Plugins/Experimental/PropertyAnimatorCore/Source/PropertyAnimatorCoreEditor/Private/Menus/PropertyAnimatorCoreEditorMenu.h
@@ -19,8 +19,6 @@ namespace UE::PropertyAnimatorCoreEditor::Menu
 {
 	/** Sections */
 
-	void FillEditAnimatorSection(UToolMenu* InMenu, TSharedRef<FPropertyAnimatorCoreEditorMenuData> InMenuData);
-
 	void FillNewAnimatorSection(UToolMenu* InMenu, TSharedRef<FPropertyAnimatorCoreEditorMenuData> InMenuData);
 
 	void FillExistingAnimatorSection(UToolMenu* InMenu, TSharedRef<FPropertyAnimatorCoreEditorMenuData> InMenuData);
@@ -45,8 +43,6 @@ namespace UE::PropertyAnimatorCoreEditor::Menu
 
 	/** Execute */
 
-	void ExecuteEditAnimatorAction(AActor* InActor);
-
 	void ExecuteNewAnimatorPresetAction(const UPropertyAnimatorCoreBase* InAnimator, const TSet<AActor*>& InActors, UPropertyAnimatorCorePresetBase* InPreset, TSharedRef<FPropertyAnimatorCoreEditorMenuData> InMenuData);
 
 	void ExecuteNewAnimatorPropertyAction(const UPropertyAnimatorCoreBase* InAnimator, FPropertyAnimatorCoreData InProperty, TSharedRef<FPropertyAnimatorCoreEditorMenuData> InMenuData);
diff --git a/Engine/Plugins/Experimental/PropertyAnimatorCore/Source/PropertyAnimatorCoreEditor/Private/Styles/PropertyAnimatorCoreEditorStyle.cpp b/Engine/Plugins/Experimental/PropertyAnimatorCore/Source/PropertyAnimatorCoreEditor/Private/Styles/PropertyAnimatorCoreEditorStyle.cpp
--- a/Engine/Plugins/Experimental/PropertyAnimatorCore/Source/PropertyAnimatorCoreEditor/Private/Styles/PropertyAnimatorCoreEditorStyle.cpp
+++ b/Engine/Plugins/Experimental/PropertyAnimatorCore/Source/PropertyAnimatorCoreEditor/Private/Styles/PropertyAnimatorCoreEditorStyle.cpp
@@ -22,11 +22,8 @@ FPropertyAnimatorCoreEditorStyle::FPropertyAnimatorCoreEditorStyle()
 
 	Set("ClassIcon.PropertyAnimatorCoreComponent",           new IMAGE_BRUSH_SVG("EditorIcons/PropertyAnimatorCoreDefault", Icon16x16));
 	Set("ClassIcon.PropertyAnimatorCoreBase",                new IMAGE_BRUSH_SVG("EditorIcons/PropertyAnimatorCoreDefault", Icon16x16));
-
 	Set("PropertyControlIcon.Default",                       new IMAGE_BRUSH_SVG("EditorIcons/PropertyAnimatorCoreDefault", Icon16x16));
 	Set("PropertyControlIcon.Linked",                        new IMAGE_BRUSH_SVG("EditorIcons/PropertyAnimatorCoreLinked",  Icon16x16));
-	Set("PropertyControlIcon.Add",                           new IMAGE_BRUSH("EditorIcons/PropertyAnimatorCoreAdd",         Icon16x16));
-	Set("PropertyControlIcon.Link",                          new IMAGE_BRUSH("EditorIcons/PropertyAnimatorCoreLink",        Icon16x16));
 
 	FSlateStyleRegistry::RegisterSlateStyle(*this);
 }
diff --git a/Engine/Plugins/Experimental/PropertyAnimatorCore/Source/PropertyAnimatorCoreEditor/Private/Subsystems/PropertyAnimatorCoreEditorSubsystem.cpp b/Engine/Plugins/Experimental/PropertyAnimatorCore/Source/PropertyAnimatorCoreEditor/Private/Subsystems/PropertyAnimatorCoreEditorSubsystem.cpp
--- a/Engine/Plugins/Experimental/PropertyAnimatorCore/Source/PropertyAnimatorCoreEditor/Private/Subsystems/PropertyAnimatorCoreEditorSubsystem.cpp
+++ b/Engine/Plugins/Experimental/PropertyAnimatorCore/Source/PropertyAnimatorCoreEditor/Private/Subsystems/PropertyAnimatorCoreEditorSubsystem.cpp
@@ -14,38 +14,9 @@
 #include "PropertyHandle.h"
 #include "Styles/PropertyAnimatorCoreEditorStyle.h"
 #include "Subsystems/PropertyAnimatorCoreSubsystem.h"
-#include "Widgets/SPropertyAnimatorCoreEditorEditPanel.h"
 
 #define LOCTEXT_NAMESPACE "PropertyAnimatorCoreEditorSubsystem"
 
-FPropertyAnimatorCoreEditorEditPanelOptions& UPropertyAnimatorCoreEditorSubsystem::OpenPropertyControlWindow()
-{
-	TSharedPtr<SPropertyAnimatorCoreEditorEditPanel> PropertyControlPanel = PropertyControllerPanelWeak.Pin();
-
-	if (!PropertyControllerPanelWeak.IsValid())
-	{
-		PropertyControlPanel = SPropertyAnimatorCoreEditorEditPanel::OpenWindow();
-		PropertyControllerPanelWeak = PropertyControlPanel;
-	}
-
-	PropertyControlPanel->FocusWindow();
-
-	return PropertyControlPanel->GetOptions();
-}
-
-void UPropertyAnimatorCoreEditorSubsystem::ClosePropertyControlWindow() const
-{
-	if (const TSharedPtr<SPropertyAnimatorCoreEditorEditPanel> PropertyControllerPanel = PropertyControllerPanelWeak.Pin())
-	{
-		PropertyControllerPanel->CloseWindow();
-	}
-}
-
-bool UPropertyAnimatorCoreEditorSubsystem::IsPropertyControlWindowOpened() const
-{
-	return PropertyControllerPanelWeak.IsValid();
-}
-
 bool UPropertyAnimatorCoreEditorSubsystem::FillAnimatorMenu(UToolMenu* InMenu, const FPropertyAnimatorCoreEditorMenuContext& InContext, const FPropertyAnimatorCoreEditorMenuOptions& InOptions)
 {
 	if (!InMenu || InContext.IsEmpty())
@@ -69,22 +40,6 @@ bool UPropertyAnimatorCoreEditorSubsystem::FillAnimatorMenu(UToolMenu* InMenu, c
 		}
 	}
 
-	if (InOptions.IsMenuType(EPropertyAnimatorCoreEditorMenuType::Edit))
-	{
-		if (InOptions.ShouldCreateSubMenu())
-		{
-			AnimatorSection->AddSubMenu(
-				TEXT("EditAnimatorMenu"),
-				LOCTEXT("EditAnimatorMenu.Label", "Edit Animators"),
-				LOCTEXT("EditAnimatorMenu.Tooltip", "Edit Animators"),
-				FNewToolMenuDelegate::CreateLambda(&UE::PropertyAnimatorCoreEditor::Menu::FillEditAnimatorSection, LastMenuData.ToSharedRef()));
-		}
-		else
-		{
-			UE::PropertyAnimatorCoreEditor::Menu::FillEditAnimatorSection(InMenu, LastMenuData.ToSharedRef());
-		}
-	}
-
 	if (InOptions.IsMenuType(EPropertyAnimatorCoreEditorMenuType::NewSimple) && InContext.ContainsAnyActor())
 	{
 		if (InOptions.ShouldCreateSubMenu())
@@ -410,7 +365,12 @@ void UPropertyAnimatorCoreEditorSubsystem::FillAnimatorExtensionSection(UToolMen
 	}
 
 	const FPropertyAnimatorCoreEditorMenuContext MenuContext({}, {Context->GetPropertyData()});
-	const FPropertyAnimatorCoreEditorMenuOptions MenuOptions({EPropertyAnimatorCoreEditorMenuType::Edit, EPropertyAnimatorCoreEditorMenuType::NewAdvanced, EPropertyAnimatorCoreEditorMenuType::Existing});
+	const FPropertyAnimatorCoreEditorMenuOptions MenuOptions(
+		{
+			EPropertyAnimatorCoreEditorMenuType::NewAdvanced
+			, EPropertyAnimatorCoreEditorMenuType::Existing
+		}
+	);
 	FillAnimatorMenu(InToolMenu, MenuContext, MenuOptions);
 }
 
@@ -437,7 +397,12 @@ void UPropertyAnimatorCoreEditorSubsystem::FillAnimatorRowContextSection(UToolMe
 	}
 
 	const FPropertyAnimatorCoreEditorMenuContext MenuContext({}, {PropertyData.GetValue()});
-	const FPropertyAnimatorCoreEditorMenuOptions MenuOptions({EPropertyAnimatorCoreEditorMenuType::Edit, EPropertyAnimatorCoreEditorMenuType::NewAdvanced, EPropertyAnimatorCoreEditorMenuType::Existing});
+	const FPropertyAnimatorCoreEditorMenuOptions MenuOptions(
+		{
+			EPropertyAnimatorCoreEditorMenuType::NewAdvanced
+			, EPropertyAnimatorCoreEditorMenuType::Existing
+		}
+	);
 	FillAnimatorMenu(InToolMenu, MenuContext, MenuOptions);
 }
 
diff --git a/Engine/Plugins/Experimental/PropertyAnimatorCore/Source/PropertyAnimatorCoreEditor/Private/Widgets/DragDropOps/PropertyAnimatorCoreEditorViewDragDropOp.h b/Engine/Plugins/Experimental/PropertyAnimatorCore/Source/PropertyAnimatorCoreEditor/Private/Widgets/DragDropOps/PropertyAnimatorCoreEditorViewDragDropOp.h
deleted file mode 100644
diff --git a/Engine/Plugins/Experimental/PropertyAnimatorCore/Source/PropertyAnimatorCoreEditor/Private/Widgets/PropertyAnimatorCoreEditorEditPanelOptions.cpp b/Engine/Plugins/Experimental/PropertyAnimatorCore/Source/PropertyAnimatorCoreEditor/Private/Widgets/PropertyAnimatorCoreEditorEditPanelOptions.cpp
deleted file mode 100644
diff --git a/Engine/Plugins/Experimental/PropertyAnimatorCore/Source/PropertyAnimatorCoreEditor/Private/Widgets/SPropertyAnimatorCoreEditorEditPanel.cpp b/Engine/Plugins/Experimental/PropertyAnimatorCore/Source/PropertyAnimatorCoreEditor/Private/Widgets/SPropertyAnimatorCoreEditorEditPanel.cpp
deleted file mode 100644
diff --git a/Engine/Plugins/Experimental/PropertyAnimatorCore/Source/PropertyAnimatorCoreEditor/Private/Widgets/SPropertyAnimatorCoreEditorEditPanel.h b/Engine/Plugins/Experimental/PropertyAnimatorCore/Source/PropertyAnimatorCoreEditor/Private/Widgets/SPropertyAnimatorCoreEditorEditPanel.h
deleted file mode 100644
diff --git a/Engine/Plugins/Experimental/PropertyAnimatorCore/Source/PropertyAnimatorCoreEditor/Private/Widgets/TableRows/SPropertyAnimatorCoreEditorControllersViewTableRow.cpp b/Engine/Plugins/Experimental/PropertyAnimatorCore/Source/PropertyAnimatorCoreEditor/Private/Widgets/TableRows/SPropertyAnimatorCoreEditorControllersViewTableRow.cpp
deleted file mode 100644
diff --git a/Engine/Plugins/Experimental/PropertyAnimatorCore/Source/PropertyAnimatorCoreEditor/Private/Widgets/TableRows/SPropertyAnimatorCoreEditorControllersViewTableRow.h b/Engine/Plugins/Experimental/PropertyAnimatorCore/Source/PropertyAnimatorCoreEditor/Private/Widgets/TableRows/SPropertyAnimatorCoreEditorControllersViewTableRow.h
deleted file mode 100644
diff --git a/Engine/Plugins/Experimental/PropertyAnimatorCore/Source/PropertyAnimatorCoreEditor/Private/Widgets/TableRows/SPropertyAnimatorCoreEditorPropertiesViewControllerTableRow.cpp b/Engine/Plugins/Experimental/PropertyAnimatorCore/Source/PropertyAnimatorCoreEditor/Private/Widgets/TableRows/SPropertyAnimatorCoreEditorPropertiesViewControllerTableRow.cpp
deleted file mode 100644
diff --git a/Engine/Plugins/Experimental/PropertyAnimatorCore/Source/PropertyAnimatorCoreEditor/Private/Widgets/TableRows/SPropertyAnimatorCoreEditorPropertiesViewControllerTableRow.h b/Engine/Plugins/Experimental/PropertyAnimatorCore/Source/PropertyAnimatorCoreEditor/Private/Widgets/TableRows/SPropertyAnimatorCoreEditorPropertiesViewControllerTableRow.h
deleted file mode 100644
diff --git a/Engine/Plugins/Experimental/PropertyAnimatorCore/Source/PropertyAnimatorCoreEditor/Private/Widgets/TableRows/SPropertyAnimatorCoreEditorPropertiesViewTableRow.cpp b/Engine/Plugins/Experimental/PropertyAnimatorCore/Source/PropertyAnimatorCoreEditor/Private/Widgets/TableRows/SPropertyAnimatorCoreEditorPropertiesViewTableRow.cpp
deleted file mode 100644
diff --git a/Engine/Plugins/Experimental/PropertyAnimatorCore/Source/PropertyAnimatorCoreEditor/Private/Widgets/TableRows/SPropertyAnimatorCoreEditorPropertiesViewTableRow.h b/Engine/Plugins/Experimental/PropertyAnimatorCore/Source/PropertyAnimatorCoreEditor/Private/Widgets/TableRows/SPropertyAnimatorCoreEditorPropertiesViewTableRow.h
deleted file mode 100644
diff --git a/Engine/Plugins/Experimental/PropertyAnimatorCore/Source/PropertyAnimatorCoreEditor/Private/Widgets/Views/PropertyAnimatorCoreEditorViewItem.h b/Engine/Plugins/Experimental/PropertyAnimatorCore/Source/PropertyAnimatorCoreEditor/Private/Widgets/Views/PropertyAnimatorCoreEditorViewItem.h
deleted file mode 100644
diff --git a/Engine/Plugins/Experimental/PropertyAnimatorCore/Source/PropertyAnimatorCoreEditor/Private/Widgets/Views/SPropertyAnimatorCoreEditorControllersView.cpp b/Engine/Plugins/Experimental/PropertyAnimatorCore/Source/PropertyAnimatorCoreEditor/Private/Widgets/Views/SPropertyAnimatorCoreEditorControllersView.cpp
deleted file mode 100644
diff --git a/Engine/Plugins/Experimental/PropertyAnimatorCore/Source/PropertyAnimatorCoreEditor/Private/Widgets/Views/SPropertyAnimatorCoreEditorControllersView.h b/Engine/Plugins/Experimental/PropertyAnimatorCore/Source/PropertyAnimatorCoreEditor/Private/Widgets/Views/SPropertyAnimatorCoreEditorControllersView.h
deleted file mode 100644
diff --git a/Engine/Plugins/Experimental/PropertyAnimatorCore/Source/PropertyAnimatorCoreEditor/Private/Widgets/Views/SPropertyAnimatorCoreEditorPropertiesView.cpp b/Engine/Plugins/Experimental/PropertyAnimatorCore/Source/PropertyAnimatorCoreEditor/Private/Widgets/Views/SPropertyAnimatorCoreEditorPropertiesView.cpp
deleted file mode 100644
diff --git a/Engine/Plugins/Experimental/PropertyAnimatorCore/Source/PropertyAnimatorCoreEditor/Private/Widgets/Views/SPropertyAnimatorCoreEditorPropertiesView.h b/Engine/Plugins/Experimental/PropertyAnimatorCore/Source/PropertyAnimatorCoreEditor/Private/Widgets/Views/SPropertyAnimatorCoreEditorPropertiesView.h
deleted file mode 100644
diff --git a/Engine/Plugins/Experimental/PropertyAnimatorCore/Source/PropertyAnimatorCoreEditor/Public/Menus/PropertyAnimatorCoreEditorMenuDefs.h b/Engine/Plugins/Experimental/PropertyAnimatorCore/Source/PropertyAnimatorCoreEditor/Public/Menus/PropertyAnimatorCoreEditorMenuDefs.h
--- a/Engine/Plugins/Experimental/PropertyAnimatorCore/Source/PropertyAnimatorCoreEditor/Public/Menus/PropertyAnimatorCoreEditorMenuDefs.h
+++ b/Engine/Plugins/Experimental/PropertyAnimatorCore/Source/PropertyAnimatorCoreEditor/Public/Menus/PropertyAnimatorCoreEditorMenuDefs.h
@@ -15,22 +15,20 @@ class UObject;
 /** Menu type available */
 enum class EPropertyAnimatorCoreEditorMenuType : uint8
 {
-	/** Menu to open edit animators window */
-	Edit = 1 << 0,
 	/** Menu to add new animator on selected actors */
-	NewSimple = 1 << 1,
+	NewSimple = 1 << 0,
 	/** Menu to add new animator on selected actors, with submenus for presets and properties */
-	NewAdvanced = 1 << 2,
+	NewAdvanced = 1 << 1,
 	/** Menu to link or unlink properties/preset to/from existing animators */
-	Existing = 1 << 3,
+	Existing = 1 << 2,
 	/** Menu to link properties/preset to existing animator */
-	Link = 1 << 4,
+	Link = 1 << 3,
 	/** Menu to delete animators */
-	Delete = 1 << 5,
+	Delete = 1 << 4,
 	/** Menu to enable animators */
-	Enable = 1 << 6,
+	Enable = 1 << 5,
 	/** Menu to disable animators */
-	Disable = 1 << 7
+	Disable = 1 << 6
 };
 
 /** Context for the menu builder, you can reuse it for multiple menus types */
diff --git a/Engine/Plugins/Experimental/PropertyAnimatorCore/Source/PropertyAnimatorCoreEditor/Public/Subsystems/PropertyAnimatorCoreEditorSubsystem.h b/Engine/Plugins/Experimental/PropertyAnimatorCore/Source/PropertyAnimatorCoreEditor/Public/Subsystems/PropertyAnimatorCoreEditorSubsystem.h
--- a/Engine/Plugins/Experimental/PropertyAnimatorCore/Source/PropertyAnimatorCoreEditor/Public/Subsystems/PropertyAnimatorCoreEditorSubsystem.h
+++ b/Engine/Plugins/Experimental/PropertyAnimatorCore/Source/PropertyAnimatorCoreEditor/Public/Subsystems/PropertyAnimatorCoreEditorSubsystem.h
@@ -20,7 +20,7 @@ struct FPropertyAnimatorCoreEditorEditPanelOptions;
 struct FOnGenerateGlobalRowExtensionArgs;
 struct FPropertyRowExtensionButton;
 
-/** Singleton class that handles editor operations for property control such as windows */
+/** Singleton class that handles editor operations for property animators */
 UCLASS()
 class UPropertyAnimatorCoreEditorSubsystem : public UEditorSubsystem
 {
@@ -35,15 +35,6 @@ public:
 	virtual void Deinitialize() override;
 	//~ End UEditorSubsystem
 
-	/** Gets or creates the property control window, single instance only allowed */
-	FPropertyAnimatorCoreEditorEditPanelOptions& OpenPropertyControlWindow();
-
-	/** Closes the property control window */
-	void ClosePropertyControlWindow() const;
-
-	/** Checks whether the single window instance is opened */
-	bool IsPropertyControlWindowOpened() const;
-
 	/** Fills a menu based on context objects and menu options */
 	PROPERTYANIMATORCOREEDITOR_API bool FillAnimatorMenu(UToolMenu* InMenu, const FPropertyAnimatorCoreEditorMenuContext& InContext, const FPropertyAnimatorCoreEditorMenuOptions& InOptions);
 
diff --git a/Engine/Plugins/Experimental/PropertyAnimatorCore/Source/PropertyAnimatorCoreEditor/Public/Widgets/PropertyAnimatorCoreEditorEditPanelOptions.h b/Engine/Plugins/Experimental/PropertyAnimatorCore/Source/PropertyAnimatorCoreEditor/Public/Widgets/PropertyAnimatorCoreEditorEditPanelOptions.h
deleted file mode 100644
