MotionDesign : PropertyAnimator
- Add sequencer animator track to drive animators with sequencer time (eval version)
- Add sequencer time source that sync with animator track when same channel is used

diff --git a/Engine/Plugins/Experimental/PropertyAnimator/Source/PropertyAnimator/Private/Animators/PropertyAnimatorFloatBase.cpp b/Engine/Plugins/Experimental/PropertyAnimator/Source/PropertyAnimator/Private/Animators/PropertyAnimatorFloatBase.cpp
--- a/Engine/Plugins/Experimental/PropertyAnimator/Source/PropertyAnimator/Private/Animators/PropertyAnimatorFloatBase.cpp
+++ b/Engine/Plugins/Experimental/PropertyAnimator/Source/PropertyAnimator/Private/Animators/PropertyAnimatorFloatBase.cpp
@@ -117,7 +117,7 @@ EPropertyAnimatorPropertySupport UPropertyAnimatorFloatBase::IsPropertySupported
 {
 	const FName TypeName = InPropertyData.GetLeafPropertyTypeName();
 
-	if (InPropertyData.IsA<FFloatProperty>())
+	if (InPropertyData.IsA<FFloatProperty>() || InPropertyData.IsA<FDoubleProperty>())
 	{
 		return EPropertyAnimatorPropertySupport::Complete;
 	}
@@ -182,7 +182,7 @@ void UPropertyAnimatorFloatBase::EvaluateProperties(FInstancedPropertyBag& InPar
 		{
 			if (FMath::Abs(PropertyTimeElapsed) > CycleDuration)
 			{
-				return false;
+				PropertyTimeElapsed = CycleDuration - UE_KINDA_SMALL_NUMBER;
 			}
 		}
 		else if (CycleMode == EPropertyAnimatorCycleMode::Loop)
diff --git a/Engine/Plugins/Experimental/PropertyAnimatorCore/Source/PropertyAnimatorCore/Private/Sequencer/MovieSceneAnimatorEvalTemplate.h b/Engine/Plugins/Experimental/PropertyAnimatorCore/Source/PropertyAnimatorCore/Private/Sequencer/MovieSceneAnimatorEvalTemplate.h
new file mode 100644
--- /dev/null
+++ b/Engine/Plugins/Experimental/PropertyAnimatorCore/Source/PropertyAnimatorCore/Private/Sequencer/MovieSceneAnimatorEvalTemplate.h
@@ -0,0 +1,35 @@
+// Copyright Epic Games, Inc. All Rights Reserved.
+
+#pragma once
+
+#include "MovieSceneAnimatorExecutionToken.h"
+#include "Evaluation/MovieSceneEvalTemplate.h"
+#include "MovieSceneAnimatorEvalTemplate.generated.h"
+
+USTRUCT()
+struct FMovieSceneAnimatorEvalTemplate : public FMovieSceneEvalTemplate
+{
+	GENERATED_BODY()
+
+	FMovieSceneAnimatorEvalTemplate() = default;
+
+	explicit FMovieSceneAnimatorEvalTemplate(uint8 InChannel)
+		: Channel(InChannel)
+	{}
+
+	//~ Begin FMovieSceneEvalTemplate
+	virtual UScriptStruct& GetScriptStructImpl() const override
+	{
+		return *StaticStruct();
+	}
+
+	virtual void Evaluate(const FMovieSceneEvaluationOperand& InOperand, const FMovieSceneContext& InContext, const FPersistentEvaluationData& InPersistentData, FMovieSceneExecutionTokens& InExecutionTokens) const override
+	{
+		FMovieSceneAnimatorExecutionToken ExecutionToken(Channel);
+		InExecutionTokens.Add(MoveTemp(ExecutionToken));
+	}
+	//~ End FMovieSceneEvalTemplate
+
+private:
+	uint8 Channel = 0;
+};
diff --git a/Engine/Plugins/Experimental/PropertyAnimatorCore/Source/PropertyAnimatorCore/Private/Sequencer/MovieSceneAnimatorExecutionToken.h b/Engine/Plugins/Experimental/PropertyAnimatorCore/Source/PropertyAnimatorCore/Private/Sequencer/MovieSceneAnimatorExecutionToken.h
new file mode 100644
--- /dev/null
+++ b/Engine/Plugins/Experimental/PropertyAnimatorCore/Source/PropertyAnimatorCore/Private/Sequencer/MovieSceneAnimatorExecutionToken.h
@@ -0,0 +1,24 @@
+// Copyright Epic Games, Inc. All Rights Reserved.
+
+#pragma once
+
+#include "Evaluation/MovieSceneExecutionTokens.h"
+#include "TimeSources/PropertyAnimatorCoreSequencerTimeSource.h"
+
+struct FMovieSceneAnimatorExecutionToken : IMovieSceneExecutionToken
+{
+	FMovieSceneAnimatorExecutionToken(uint8 InChannel)
+		: Channel(InChannel)
+	{}
+
+	//~ Begin IMovieSceneExecutionToken
+	virtual void Execute(const FMovieSceneContext& InContext, const FMovieSceneEvaluationOperand& InOperand, FPersistentEvaluationData& InPersistentData, IMovieScenePlayer& InPlayer) override
+	{
+		const double EvaluatedTime = InContext.GetTime().AsDecimal() / InContext.GetFrameRate().AsDecimal();
+		UPropertyAnimatorCoreSequencerTimeSource::OnAnimatorTimeEvaluated.Broadcast(Channel, EvaluatedTime);
+	}
+	//~ End IMovieSceneExecutionToken
+
+private:
+	uint8 Channel = 0;
+};
diff --git a/Engine/Plugins/Experimental/PropertyAnimatorCore/Source/PropertyAnimatorCore/Private/Sequencer/MovieSceneAnimatorSection.cpp b/Engine/Plugins/Experimental/PropertyAnimatorCore/Source/PropertyAnimatorCore/Private/Sequencer/MovieSceneAnimatorSection.cpp
new file mode 100644
--- /dev/null
+++ b/Engine/Plugins/Experimental/PropertyAnimatorCore/Source/PropertyAnimatorCore/Private/Sequencer/MovieSceneAnimatorSection.cpp
@@ -0,0 +1,9 @@
+// Copyright Epic Games, Inc. All Rights Reserved.
+
+#include "Sequencer/MovieSceneAnimatorSection.h"
+
+UMovieSceneAnimatorSection::UMovieSceneAnimatorSection()
+	: UMovieSceneSection()
+{
+	bSupportsInfiniteRange = true;
+}
diff --git a/Engine/Plugins/Experimental/PropertyAnimatorCore/Source/PropertyAnimatorCore/Private/Sequencer/MovieSceneAnimatorSection.h b/Engine/Plugins/Experimental/PropertyAnimatorCore/Source/PropertyAnimatorCore/Private/Sequencer/MovieSceneAnimatorSection.h
new file mode 100644
--- /dev/null
+++ b/Engine/Plugins/Experimental/PropertyAnimatorCore/Source/PropertyAnimatorCore/Private/Sequencer/MovieSceneAnimatorSection.h
@@ -0,0 +1,26 @@
+// Copyright Epic Games, Inc. All Rights Reserved.
+
+#pragma once
+
+#include "MovieSceneSection.h"
+#include "MovieSceneAnimatorSection.generated.h"
+
+/** Movie scene section for a sequencer animator channel */
+UCLASS(MinimalAPI)
+class UMovieSceneAnimatorSection : public UMovieSceneSection
+{
+	GENERATED_BODY()
+
+public:
+	UMovieSceneAnimatorSection();
+
+	uint8 GetChannel() const
+	{
+		return Channel;
+	}
+
+protected:
+	/** Channel used to push sequencer time to */
+	UPROPERTY(EditAnywhere, Category="Animator")
+	uint8 Channel = 0;
+};
diff --git a/Engine/Plugins/Experimental/PropertyAnimatorCore/Source/PropertyAnimatorCore/Private/Sequencer/MovieSceneAnimatorTrack.cpp b/Engine/Plugins/Experimental/PropertyAnimatorCore/Source/PropertyAnimatorCore/Private/Sequencer/MovieSceneAnimatorTrack.cpp
new file mode 100644
--- /dev/null
+++ b/Engine/Plugins/Experimental/PropertyAnimatorCore/Source/PropertyAnimatorCore/Private/Sequencer/MovieSceneAnimatorTrack.cpp
@@ -0,0 +1,98 @@
+// Copyright Epic Games, Inc. All Rights Reserved.
+
+#include "Sequencer/MovieSceneAnimatorTrack.h"
+
+#include "Evaluation/Blending/MovieSceneBlendType.h"
+#include "Evaluation/MovieSceneEvalTemplate.h"
+#include "MovieScene.h"
+#include "Sequencer/MovieSceneAnimatorEvalTemplate.h"
+#include "Sequencer/MovieSceneAnimatorSection.h"
+
+#define LOCTEXT_NAMESPACE "MovieSceneAnimatorTrack"
+
+UMovieSceneAnimatorTrack::UMovieSceneAnimatorTrack()
+	: UMovieSceneNameableTrack()
+{
+	SupportedBlendTypes = FMovieSceneBlendTypeField::None();
+}
+
+bool UMovieSceneAnimatorTrack::SupportsType(TSubclassOf<UMovieSceneSection> InSectionClass) const
+{
+	return InSectionClass == UMovieSceneAnimatorSection::StaticClass();
+}
+
+UMovieSceneSection* UMovieSceneAnimatorTrack::CreateNewSection()
+{
+	UMovieSceneAnimatorSection* NewSection = NewObject<UMovieSceneAnimatorSection>(this, NAME_None, RF_Transactional);
+
+	if (const UMovieScene* MovieScene = GetTypedOuter<UMovieScene>())
+	{
+		NewSection->SetStartFrame(MovieScene->GetPlaybackRange().GetLowerBound());
+		NewSection->SetEndFrame(MovieScene->GetPlaybackRange().GetUpperBound());
+	}
+
+	return NewSection;
+}
+
+void UMovieSceneAnimatorTrack::AddSection(UMovieSceneSection& InSection)
+{
+	Sections.Add(&InSection);
+}
+
+const TArray<UMovieSceneSection*>& UMovieSceneAnimatorTrack::GetAllSections() const
+{
+	return Sections;
+}
+
+bool UMovieSceneAnimatorTrack::HasSection(const UMovieSceneSection& InSection) const
+{
+	return Sections.Contains(&InSection);
+}
+
+bool UMovieSceneAnimatorTrack::IsEmpty() const
+{
+	return Sections.IsEmpty();
+}
+
+void UMovieSceneAnimatorTrack::RemoveAllAnimationData()
+{
+	Sections.Empty();
+}
+
+void UMovieSceneAnimatorTrack::RemoveSection(UMovieSceneSection& InSection)
+{
+	Sections.Remove(&InSection);
+}
+
+void UMovieSceneAnimatorTrack::RemoveSectionAt(int32 InSectionIndex)
+{
+	Sections.RemoveAt(InSectionIndex);
+}
+
+#if WITH_EDITORONLY_DATA
+FText UMovieSceneAnimatorTrack::GetDefaultDisplayName() const
+{
+	return LOCTEXT("MovieSceneAnimatorTrackName", "Animator Channel");
+}
+
+bool UMovieSceneAnimatorTrack::CanRename() const
+{
+	return true;
+}
+#endif
+
+FMovieSceneEvalTemplatePtr UMovieSceneAnimatorTrack::CreateTemplateForSection(const UMovieSceneSection& InSection) const
+{
+	const UMovieSceneAnimatorSection* AnimatorSection = Cast<UMovieSceneAnimatorSection>(&InSection);
+
+	if (!AnimatorSection)
+	{
+		return FMovieSceneEvalTemplatePtr();
+	}
+
+	return FMovieSceneAnimatorEvalTemplate(AnimatorSection->GetChannel());
+}
+
+#undef LOCTEXT_NAMESPACE
+
+
diff --git a/Engine/Plugins/Experimental/PropertyAnimatorCore/Source/PropertyAnimatorCore/Private/TimeSources/PropertyAnimatorCoreSequencerTimeSource.cpp b/Engine/Plugins/Experimental/PropertyAnimatorCore/Source/PropertyAnimatorCore/Private/TimeSources/PropertyAnimatorCoreSequencerTimeSource.cpp
new file mode 100644
--- /dev/null
+++ b/Engine/Plugins/Experimental/PropertyAnimatorCore/Source/PropertyAnimatorCore/Private/TimeSources/PropertyAnimatorCoreSequencerTimeSource.cpp
@@ -0,0 +1,42 @@
+// Copyright Epic Games, Inc. All Rights Reserved.
+
+#include "TimeSources/PropertyAnimatorCoreSequencerTimeSource.h"
+
+UPropertyAnimatorCoreSequencerTimeSource::FOnAnimatorTimeEvaluated UPropertyAnimatorCoreSequencerTimeSource::OnAnimatorTimeEvaluated;
+
+double UPropertyAnimatorCoreSequencerTimeSource::GetTimeElapsed()
+{
+	return EvalTime.Get(0);
+}
+
+bool UPropertyAnimatorCoreSequencerTimeSource::IsTimeSourceReady() const
+{
+	return true;
+}
+
+void UPropertyAnimatorCoreSequencerTimeSource::OnTimeSourceActive()
+{
+	Super::OnTimeSourceActive();
+
+	OnAnimatorTimeEvaluated.AddUObject(this, &UPropertyAnimatorCoreSequencerTimeSource::OnSequencerTimeEvaluated);
+}
+
+void UPropertyAnimatorCoreSequencerTimeSource::OnTimeSourceInactive()
+{
+	Super::OnTimeSourceInactive();
+
+	OnAnimatorTimeEvaluated.RemoveAll(this);
+}
+
+void UPropertyAnimatorCoreSequencerTimeSource::SetChannel(uint8 InChannel)
+{
+	Channel = InChannel;
+}
+
+void UPropertyAnimatorCoreSequencerTimeSource::OnSequencerTimeEvaluated(uint8 InChannel, double InTimeEval)
+{
+	if (Channel == InChannel)
+	{
+		EvalTime = InTimeEval;
+	}
+}
diff --git a/Engine/Plugins/Experimental/PropertyAnimatorCore/Source/PropertyAnimatorCore/PropertyAnimatorCore.Build.cs b/Engine/Plugins/Experimental/PropertyAnimatorCore/Source/PropertyAnimatorCore/PropertyAnimatorCore.Build.cs
--- a/Engine/Plugins/Experimental/PropertyAnimatorCore/Source/PropertyAnimatorCore/PropertyAnimatorCore.Build.cs
+++ b/Engine/Plugins/Experimental/PropertyAnimatorCore/Source/PropertyAnimatorCore/PropertyAnimatorCore.Build.cs
@@ -12,7 +12,8 @@ public class PropertyAnimatorCore : ModuleRules
 			new string[]
 			{
 				"Core",
-				"CoreUObject"
+				"CoreUObject",
+				"MovieScene"
 			});
 
 		PrivateDependencyModuleNames.AddRange(
@@ -20,7 +21,7 @@ public class PropertyAnimatorCore : ModuleRules
 			{
 				"ApplicationCore",
 				"CoreUObject",
-				"Engine",
+				"Engine"
 			});
 
 		if (Target.Type == TargetRules.TargetType.Editor)
diff --git a/Engine/Plugins/Experimental/PropertyAnimatorCore/Source/PropertyAnimatorCore/Public/Sequencer/MovieSceneAnimatorTrack.h b/Engine/Plugins/Experimental/PropertyAnimatorCore/Source/PropertyAnimatorCore/Public/Sequencer/MovieSceneAnimatorTrack.h
new file mode 100644
--- /dev/null
+++ b/Engine/Plugins/Experimental/PropertyAnimatorCore/Source/PropertyAnimatorCore/Public/Sequencer/MovieSceneAnimatorTrack.h
@@ -0,0 +1,41 @@
+// Copyright Epic Games, Inc. All Rights Reserved.
+
+#pragma once
+
+#include "Compilation/IMovieSceneTrackTemplateProducer.h"
+#include "MovieSceneNameableTrack.h"
+#include "MovieSceneAnimatorTrack.generated.h"
+
+/** Movie scene track used to drive animator with sequencer */
+UCLASS(MinimalAPI)
+class UMovieSceneAnimatorTrack : public UMovieSceneNameableTrack, public IMovieSceneTrackTemplateProducer
+{
+	GENERATED_BODY()
+
+public:
+	UMovieSceneAnimatorTrack();
+
+	//~ Begin UMovieSceneTrack
+	virtual bool SupportsType(TSubclassOf<UMovieSceneSection> InSectionClass) const override;
+	virtual UMovieSceneSection* CreateNewSection() override;
+	virtual void AddSection(UMovieSceneSection& InSection) override;
+	virtual const TArray<UMovieSceneSection*>& GetAllSections() const override;
+	virtual bool HasSection(const UMovieSceneSection& InSection) const override;
+	virtual bool IsEmpty() const override;
+	virtual void RemoveAllAnimationData() override;
+	virtual void RemoveSection(UMovieSceneSection& InSection) override;
+	virtual void RemoveSectionAt(int32 InSectionIndex) override;
+#if WITH_EDITORONLY_DATA
+	virtual FText GetDefaultDisplayName() const override;
+	virtual bool CanRename() const override;
+#endif
+	//~ End UMovieSceneTrack
+
+	//~ Begin IMovieSceneTrackTemplateProducer
+	virtual FMovieSceneEvalTemplatePtr CreateTemplateForSection(const UMovieSceneSection& InSection) const override;
+	//~ End IMovieSceneTrackTemplateProducer
+
+private:
+	UPROPERTY()
+	TArray<TObjectPtr<UMovieSceneSection>> Sections;
+};
diff --git a/Engine/Plugins/Experimental/PropertyAnimatorCore/Source/PropertyAnimatorCore/Public/TimeSources/PropertyAnimatorCoreSequencerTimeSource.h b/Engine/Plugins/Experimental/PropertyAnimatorCore/Source/PropertyAnimatorCore/Public/TimeSources/PropertyAnimatorCoreSequencerTimeSource.h
new file mode 100644
--- /dev/null
+++ b/Engine/Plugins/Experimental/PropertyAnimatorCore/Source/PropertyAnimatorCore/Public/TimeSources/PropertyAnimatorCoreSequencerTimeSource.h
@@ -0,0 +1,44 @@
+// Copyright Epic Games, Inc. All Rights Reserved.
+
+#pragma once
+
+#include "PropertyAnimatorCoreTimeSourceBase.h"
+#include "PropertyAnimatorCoreSequencerTimeSource.generated.h"
+
+/** Sequencer time source that sync with animator track channel */
+UCLASS(MinimalAPI)
+class UPropertyAnimatorCoreSequencerTimeSource : public UPropertyAnimatorCoreTimeSourceBase
+{
+	GENERATED_BODY()
+
+public:
+	DECLARE_MULTICAST_DELEGATE_TwoParams(FOnAnimatorTimeEvaluated, uint8 /** Channel */, double /** EvalTime */)
+	static FOnAnimatorTimeEvaluated OnAnimatorTimeEvaluated;
+
+	UPropertyAnimatorCoreSequencerTimeSource()
+		: UPropertyAnimatorCoreTimeSourceBase(TEXT("Sequencer"))
+	{}
+
+	//~ Begin UPropertyAnimatorTimeSourceBase
+	virtual double GetTimeElapsed() override;
+	virtual bool IsTimeSourceReady() const override;
+	virtual void OnTimeSourceActive() override;
+	virtual void OnTimeSourceInactive() override;
+	//~ End UPropertyAnimatorTimeSourceBase
+
+	void SetChannel(uint8 InChannel);
+	uint8 GetChannel() const
+	{
+		return Channel;
+	}
+
+protected:
+	void OnSequencerTimeEvaluated(uint8 InChannel, double InTimeEval);
+
+	/** Channel to sample time from */
+	UPROPERTY(EditInstanceOnly, Setter, Getter, Category="Animator")
+	uint8 Channel = 0;
+
+	/** Last evaluated time received */
+	TOptional<double> EvalTime;
+};
\ No newline at end of file
diff --git a/Engine/Plugins/Experimental/PropertyAnimatorCore/Source/PropertyAnimatorCoreEditor/Private/PropertyAnimatorCoreEditorModule.cpp b/Engine/Plugins/Experimental/PropertyAnimatorCore/Source/PropertyAnimatorCoreEditor/Private/PropertyAnimatorCoreEditorModule.cpp
--- a/Engine/Plugins/Experimental/PropertyAnimatorCore/Source/PropertyAnimatorCoreEditor/Private/PropertyAnimatorCoreEditorModule.cpp
+++ b/Engine/Plugins/Experimental/PropertyAnimatorCore/Source/PropertyAnimatorCoreEditor/Private/PropertyAnimatorCoreEditorModule.cpp
@@ -5,15 +5,20 @@
 #include "Animators/PropertyAnimatorCoreBase.h"
 #include "Customizations/PropertyAnimatorCoreEditorDetailCustomization.h"
 #include "Customizations/PropertyAnimatorCoreEditorContextTypeCustomization.h"
+#include "ISequencerModule.h"
 #include "Modules/ModuleManager.h"
 #include "Properties/PropertyAnimatorCoreContext.h"
 #include "PropertyEditorModule.h"
+#include "Sequencer/MovieSceneAnimatorTrackEditor.h"
 
 void FPropertyAnimatorCoreEditorModule::StartupModule()
 {
 	FPropertyEditorModule& PropertyModule = FModuleManager::GetModuleChecked<FPropertyEditorModule>("PropertyEditor");
 	PropertyModule.RegisterCustomPropertyTypeLayout(UPropertyAnimatorCoreContext::StaticClass()->GetFName(), FOnGetPropertyTypeCustomizationInstance::CreateStatic(&FPropertyAnimatorCoreEditorContextTypeCustomization::MakeInstance));
 	PropertyModule.RegisterCustomClassLayout(UPropertyAnimatorCoreBase::StaticClass()->GetFName(), FOnGetDetailCustomizationInstance::CreateStatic(&FPropertyAnimatorCoreEditorDetailCustomization::MakeInstance));
+
+	ISequencerModule& SequencerModule = FModuleManager::Get().LoadModuleChecked<ISequencerModule>("Sequencer");
+	AnimatorTrackCreateEditorHandle = SequencerModule.RegisterTrackEditor(FOnCreateTrackEditor::CreateStatic(&FMovieSceneAnimatorTrackEditor::CreateTrackEditor));
 }
 
 void FPropertyAnimatorCoreEditorModule::ShutdownModule()
@@ -24,6 +29,13 @@ void FPropertyAnimatorCoreEditorModule::ShutdownModule()
 		PropertyModule.UnregisterCustomPropertyTypeLayout(UPropertyAnimatorCoreContext::StaticClass()->GetFName());
 		PropertyModule.UnregisterCustomClassLayout(UPropertyAnimatorCoreBase::StaticClass()->GetFName());
 	}
+
+	if (FModuleManager::Get().IsModuleLoaded("Sequencer"))
+	{
+		ISequencerModule& SequencerModule = FModuleManager::GetModuleChecked<ISequencerModule>("Sequencer");
+		SequencerModule.UnRegisterTrackEditor(AnimatorTrackCreateEditorHandle);
+		AnimatorTrackCreateEditorHandle.Reset();
+	}
 }
 
 IMPLEMENT_MODULE(FPropertyAnimatorCoreEditorModule, PropertyAnimatorCoreEditor)
diff --git a/Engine/Plugins/Experimental/PropertyAnimatorCore/Source/PropertyAnimatorCoreEditor/Private/PropertyAnimatorCoreEditorModule.h b/Engine/Plugins/Experimental/PropertyAnimatorCore/Source/PropertyAnimatorCoreEditor/Private/PropertyAnimatorCoreEditorModule.h
--- a/Engine/Plugins/Experimental/PropertyAnimatorCore/Source/PropertyAnimatorCoreEditor/Private/PropertyAnimatorCoreEditorModule.h
+++ b/Engine/Plugins/Experimental/PropertyAnimatorCore/Source/PropertyAnimatorCoreEditor/Private/PropertyAnimatorCoreEditorModule.h
@@ -11,4 +11,7 @@ public:
 	virtual void StartupModule() override;
 	virtual void ShutdownModule() override;
 	//~ End IModuleInterface
+
+private:
+	FDelegateHandle AnimatorTrackCreateEditorHandle;
 };
\ No newline at end of file
diff --git a/Engine/Plugins/Experimental/PropertyAnimatorCore/Source/PropertyAnimatorCoreEditor/Private/Sequencer/MovieSceneAnimatorTrackEditor.cpp b/Engine/Plugins/Experimental/PropertyAnimatorCore/Source/PropertyAnimatorCoreEditor/Private/Sequencer/MovieSceneAnimatorTrackEditor.cpp
new file mode 100644
--- /dev/null
+++ b/Engine/Plugins/Experimental/PropertyAnimatorCore/Source/PropertyAnimatorCoreEditor/Private/Sequencer/MovieSceneAnimatorTrackEditor.cpp
@@ -0,0 +1,46 @@
+// Copyright Epic Games, Inc. All Rights Reserved.
+
+#include "MovieSceneAnimatorTrackEditor.h"
+
+#include "Styling/SlateIconFinder.h"
+
+#define LOCTEXT_NAMESPACE "MovieSceneAnimatorTrackEditor"
+
+void FMovieSceneAnimatorTrackEditor::BuildAddTrackMenu(FMenuBuilder& InMenuBuilder)
+{
+	InMenuBuilder.AddMenuEntry(
+		LOCTEXT("AddAnimatorTrack.Label", "Animator"),
+		LOCTEXT("AddAnimatorTrack.Tooltip", "Adds a new track that uses the time of the current sequence to drive animators."),
+		FSlateIconFinder::FindIconForClass(UMovieSceneAnimatorTrack::StaticClass()),
+		FUIAction(
+			FExecuteAction::CreateSP(this, &FMovieSceneAnimatorTrackEditor::ExecuteAddTrack)
+		)
+	);
+}
+
+void FMovieSceneAnimatorTrackEditor::ExecuteAddTrack()
+{
+	UMovieScene* FocusedMovieScene = GetFocusedMovieScene();
+	if (!FocusedMovieScene || FocusedMovieScene->IsReadOnly())
+	{
+		return;
+	}
+
+	const TSharedPtr<ISequencer> SequencerPtr = GetSequencer();
+	if (!SequencerPtr)
+	{
+		return;
+	}
+
+	const FScopedTransaction Transaction(LOCTEXT("AddAnimatorTrack", "Add Animator Track"));
+
+	FocusedMovieScene->Modify();
+
+	UMovieSceneAnimatorTrack* NewTrack = NewObject<UMovieSceneAnimatorTrack>(FocusedMovieScene, NAME_None, RF_Transactional);
+	NewTrack->AddSection(*NewTrack->CreateNewSection());
+
+	FocusedMovieScene->AddGivenTrack(NewTrack);
+	SequencerPtr->OnAddTrack(NewTrack, FGuid());
+}
+
+#undef LOCTEXT_NAMESPACE
diff --git a/Engine/Plugins/Experimental/PropertyAnimatorCore/Source/PropertyAnimatorCoreEditor/Private/Sequencer/MovieSceneAnimatorTrackEditor.h b/Engine/Plugins/Experimental/PropertyAnimatorCore/Source/PropertyAnimatorCoreEditor/Private/Sequencer/MovieSceneAnimatorTrackEditor.h
new file mode 100644
--- /dev/null
+++ b/Engine/Plugins/Experimental/PropertyAnimatorCore/Source/PropertyAnimatorCoreEditor/Private/Sequencer/MovieSceneAnimatorTrackEditor.h
@@ -0,0 +1,26 @@
+// Copyright Epic Games, Inc. All Rights Reserved.
+
+#pragma once
+
+#include "KeyframeTrackEditor.h"
+#include "Sequencer/MovieSceneAnimatorTrack.h"
+
+class FMovieSceneAnimatorTrackEditor : public FKeyframeTrackEditor<UMovieSceneAnimatorTrack>
+{
+public:
+	FMovieSceneAnimatorTrackEditor(const TSharedRef<ISequencer>& InSequencer)
+		: FKeyframeTrackEditor<UMovieSceneAnimatorTrack>(InSequencer)
+	{}
+
+	static TSharedRef<ISequencerTrackEditor> CreateTrackEditor(TSharedRef<ISequencer> InSequencer)
+	{
+		return MakeShared<FMovieSceneAnimatorTrackEditor>(InSequencer);
+	}
+
+private:
+	//~ Begin FMovieSceneTrackEditor
+	virtual void BuildAddTrackMenu(FMenuBuilder& InMenuBuilder) override;
+	//~ End FMovieSceneTrackEditor
+
+	void ExecuteAddTrack();
+};
diff --git a/Engine/Plugins/Experimental/PropertyAnimatorCore/Source/PropertyAnimatorCoreEditor/Private/Styles/PropertyAnimatorCoreEditorStyle.cpp b/Engine/Plugins/Experimental/PropertyAnimatorCore/Source/PropertyAnimatorCoreEditor/Private/Styles/PropertyAnimatorCoreEditorStyle.cpp
--- a/Engine/Plugins/Experimental/PropertyAnimatorCore/Source/PropertyAnimatorCoreEditor/Private/Styles/PropertyAnimatorCoreEditorStyle.cpp
+++ b/Engine/Plugins/Experimental/PropertyAnimatorCore/Source/PropertyAnimatorCoreEditor/Private/Styles/PropertyAnimatorCoreEditorStyle.cpp
@@ -20,10 +20,11 @@ FPropertyAnimatorCoreEditorStyle::FPropertyAnimatorCoreEditorStyle()
 
 	ContentRootDir = FPaths::Combine(Plugin->GetBaseDir(), TEXT("Resources"));
 
-	Set("ClassIcon.PropertyAnimatorCoreComponent",           new IMAGE_BRUSH_SVG("EditorIcons/PropertyAnimatorCoreDefault", Icon16x16));
-	Set("ClassIcon.PropertyAnimatorCoreBase",                new IMAGE_BRUSH_SVG("EditorIcons/PropertyAnimatorCoreDefault", Icon16x16));
-	Set("PropertyControlIcon.Default",                       new IMAGE_BRUSH_SVG("EditorIcons/PropertyAnimatorCoreDefault", Icon16x16));
-	Set("PropertyControlIcon.Linked",                        new IMAGE_BRUSH_SVG("EditorIcons/PropertyAnimatorCoreLinked",  Icon16x16));
+	Set("ClassIcon.MovieSceneAnimatorTrack",        new IMAGE_BRUSH_SVG("EditorIcons/PropertyAnimatorCoreDefault", Icon16x16));
+	Set("ClassIcon.PropertyAnimatorCoreComponent",  new IMAGE_BRUSH_SVG("EditorIcons/PropertyAnimatorCoreDefault", Icon16x16));
+	Set("ClassIcon.PropertyAnimatorCoreBase",       new IMAGE_BRUSH_SVG("EditorIcons/PropertyAnimatorCoreDefault", Icon16x16));
+	Set("PropertyControlIcon.Default",              new IMAGE_BRUSH_SVG("EditorIcons/PropertyAnimatorCoreDefault", Icon16x16));
+	Set("PropertyControlIcon.Linked",               new IMAGE_BRUSH_SVG("EditorIcons/PropertyAnimatorCoreLinked",  Icon16x16));
 
 	FSlateStyleRegistry::RegisterSlateStyle(*this);
 }
diff --git a/Engine/Plugins/Experimental/PropertyAnimatorCore/Source/PropertyAnimatorCoreEditor/PropertyAnimatorCoreEditor.Build.cs b/Engine/Plugins/Experimental/PropertyAnimatorCore/Source/PropertyAnimatorCoreEditor/PropertyAnimatorCoreEditor.Build.cs
--- a/Engine/Plugins/Experimental/PropertyAnimatorCore/Source/PropertyAnimatorCoreEditor/PropertyAnimatorCoreEditor.Build.cs
+++ b/Engine/Plugins/Experimental/PropertyAnimatorCore/Source/PropertyAnimatorCoreEditor/PropertyAnimatorCoreEditor.Build.cs
@@ -24,10 +24,13 @@ public class PropertyAnimatorCoreEditor : ModuleRules
 		        "EditorWidgets",
 		        "Engine",
 		        "InputCore",
+				"MovieScene",
+				"MovieSceneTools",
 		        "OperatorStackEditor",
 		        "Projects",
 				"PropertyAnimatorCore",
 				"PropertyEditor",
+				"Sequencer",
 		        "SlateCore",
 		        "Slate",
 		        "ToolMenus",
