MotionDesign : PropertyAnimator
- Rebinding to linked property even when the owner name or class is different (during copy/move operations)
- Increase property search depth to drive struct property containing supported property

diff --git a/Engine/Plugins/Experimental/PropertyAnimator/Source/PropertyAnimator/Public/Properties/PropertyAnimatorVectorContext.h b/Engine/Plugins/Experimental/PropertyAnimator/Source/PropertyAnimator/Public/Properties/PropertyAnimatorVectorContext.h
--- a/Engine/Plugins/Experimental/PropertyAnimator/Source/PropertyAnimator/Public/Properties/PropertyAnimatorVectorContext.h
+++ b/Engine/Plugins/Experimental/PropertyAnimator/Source/PropertyAnimator/Public/Properties/PropertyAnimatorVectorContext.h
@@ -39,7 +39,7 @@ protected:
 	FVector GetClampedAmplitude(FVector InAmplitude);
 
 	/** The minimum value should be remapped to that values */
-	UPROPERTY(EditInstanceOnly, Setter, Getter, Category="Animator", meta=(MotionDesignVectorWidget, AllowPreserveRatio="XYZ"))
+	UPROPERTY(EditInstanceOnly, Setter, Getter, Category="Animator", meta=(AllowPreserveRatio))
 	FVector AmplitudeMin = FVector::ZeroVector;
 
 	/** Some properties are clamped and cannot go past a specific min value */
@@ -47,7 +47,7 @@ protected:
 	TOptional<FVector> AmplitudeClampMin;
 
 	/** The maximum value should be remapped to that values */
-	UPROPERTY(EditInstanceOnly, Setter, Getter, Category="Animator", meta=(MotionDesignVectorWidget, AllowPreserveRatio="XYZ"))
+	UPROPERTY(EditInstanceOnly, Setter, Getter, Category="Animator", meta=(AllowPreserveRatio))
 	FVector AmplitudeMax = FVector::ZeroVector;
 
 	/** Some properties are clamped and cannot go past a specific min value */
diff --git a/Engine/Plugins/Experimental/PropertyAnimatorCore/Source/PropertyAnimatorCore/Private/Animators/PropertyAnimatorCoreBase.cpp b/Engine/Plugins/Experimental/PropertyAnimatorCore/Source/PropertyAnimatorCore/Private/Animators/PropertyAnimatorCoreBase.cpp
--- a/Engine/Plugins/Experimental/PropertyAnimatorCore/Source/PropertyAnimatorCore/Private/Animators/PropertyAnimatorCoreBase.cpp
+++ b/Engine/Plugins/Experimental/PropertyAnimatorCore/Source/PropertyAnimatorCore/Private/Animators/PropertyAnimatorCoreBase.cpp
@@ -342,7 +342,9 @@ TSubclassOf<UPropertyAnimatorCoreContext> UPropertyAnimatorCoreBase::GetProperty
 
 void UPropertyAnimatorCoreBase::OnAnimatorEnabledChanged()
 {
-	if (bAnimatorEnabled)
+	const UPropertyAnimatorCoreComponent* AnimatorComponent = GetAnimatorComponent();
+
+	if (bAnimatorEnabled && AnimatorComponent->ShouldAnimate())
 	{
 		OnAnimatorEnabled();
 	}
@@ -427,7 +429,7 @@ void UPropertyAnimatorCoreBase::EvaluateAnimator(FInstancedPropertyBag& InParame
 
 	EvaluatedPropertyValues.Reset();
 	InParameters.AddProperty(TimeElapsedParameterName, EPropertyBagPropertyType::Double);
-	InParameters.SetValueDouble(TimeElapsedParameterName, TimeElapsed.GetValue());
+    InParameters.SetValueDouble(TimeElapsedParameterName, TimeElapsed.GetValue());
 
 	bEvaluatingProperties = true;
 	EvaluateProperties(InParameters);
diff --git a/Engine/Plugins/Experimental/PropertyAnimatorCore/Source/PropertyAnimatorCore/Private/Components/PropertyAnimatorCoreComponent.cpp b/Engine/Plugins/Experimental/PropertyAnimatorCore/Source/PropertyAnimatorCore/Private/Components/PropertyAnimatorCoreComponent.cpp
--- a/Engine/Plugins/Experimental/PropertyAnimatorCore/Source/PropertyAnimatorCore/Private/Components/PropertyAnimatorCoreComponent.cpp
+++ b/Engine/Plugins/Experimental/PropertyAnimatorCore/Source/PropertyAnimatorCore/Private/Components/PropertyAnimatorCoreComponent.cpp
@@ -98,7 +98,6 @@ void UPropertyAnimatorCoreComponent::OnAnimatorsChanged()
 		}
 	}
 
-	const bool bShouldAnimatorsTick = ShouldAnimatorsTick();
 	for (const TObjectPtr<UPropertyAnimatorCoreBase>& AddedAnimator : AddedAnimators)
 	{
 		if (AddedAnimator)
@@ -106,11 +105,6 @@ void UPropertyAnimatorCoreComponent::OnAnimatorsChanged()
 			AddedAnimator->SetAnimatorDisplayName(GetAnimatorName(AddedAnimator));
 			AddedAnimator->OnAnimatorAdded();
 			AddedAnimator->SetAnimatorEnabled(true);
-
-			if (!bShouldAnimatorsTick)
-			{
-				AddedAnimator->SetAnimatorEnabled(false);
-			}
 		}
 	}
 
@@ -119,12 +113,7 @@ void UPropertyAnimatorCoreComponent::OnAnimatorsChanged()
 
 void UPropertyAnimatorCoreComponent::OnAnimatorsEnabledChanged()
 {
-	const bool bEnableAnimators = ShouldAnimatorsTick();
-
-	if (IsComponentTickEnabled() == bEnableAnimators)
-	{
-		return;
-	}
+	const bool bEnableAnimators = ShouldAnimate();
 
 	for (const TObjectPtr<UPropertyAnimatorCoreBase>& Animator : PropertyAnimators)
 	{
@@ -133,20 +122,25 @@ void UPropertyAnimatorCoreComponent::OnAnimatorsEnabledChanged()
 			continue;
 		}
 
-		if (bEnableAnimators)
+		// When enabling all animators, if animator is disabled then skip
+		if (bEnableAnimators && !Animator->GetAnimatorEnabled())
 		{
-			Animator->OnAnimatorEnabled();
+			continue;
 		}
-		else
+
+		// When disabling all animators : if animator is disabled then skip
+		if (!bEnableAnimators && !Animator->GetAnimatorEnabled())
 		{
-			Animator->OnAnimatorDisabled();
+			continue;
 		}
+
+		Animator->OnAnimatorEnabledChanged();
 	}
 
 	SetComponentTickEnabled(bEnableAnimators);
 }
 
-bool UPropertyAnimatorCoreComponent::ShouldAnimatorsTick() const
+bool UPropertyAnimatorCoreComponent::ShouldAnimate() const
 {
 	return bAnimatorsEnabled
 		&& !PropertyAnimators.IsEmpty()
@@ -191,40 +185,6 @@ void UPropertyAnimatorCoreComponent::OnComponentCreated()
 	}
 }
 
-void UPropertyAnimatorCoreComponent::TickComponent(float InDeltaTime, ELevelTick InTickType, FActorComponentTickFunction* InThisTickFunction)
-{
-	Super::TickComponent(InDeltaTime, InTickType, InThisTickFunction);
-
-	if (!ShouldAnimatorsTick())
-	{
-		return;
-	}
-
-	const UWorld* World = GetWorld();
-	const bool bIsSupportedWorld = IsValid(World) && (World->IsGameWorld() || World->IsEditorWorld());
-
-	if (!bIsSupportedWorld)
-	{
-		return;
-	}
-
-	FInstancedPropertyBag Parameters;
-
-	for (const TObjectPtr<UPropertyAnimatorCoreBase>& Animator : PropertyAnimators)
-	{
-		if (!IsValid(Animator))
-		{
-			continue;
-		}
-
-		Parameters.Reset();
-		Parameters.AddProperty(UPropertyAnimatorCoreBase::MagnitudeParameterName, EPropertyBagPropertyType::Float);
-		Parameters.SetValueFloat(UPropertyAnimatorCoreBase::MagnitudeParameterName, AnimatorsMagnitude);
-
-		Animator->EvaluateAnimator(Parameters);
-	}
-}
-
 UPropertyAnimatorCoreComponent* UPropertyAnimatorCoreComponent::FindOrAdd(AActor* InActor)
 {
 	if (!IsValid(InActor))
@@ -268,7 +228,6 @@ UPropertyAnimatorCoreComponent::UPropertyAnimatorCoreComponent()
 	{
 		bTickInEditor = true;
 		PrimaryComponentTick.bCanEverTick = true;
-		PrimaryComponentTick.bHighPriority = true;
 
 		// Used to toggle animators state in world
 		UPropertyAnimatorCoreSubsystem::OnAnimatorsSetEnabledDelegate.AddUObject(this, &UPropertyAnimatorCoreComponent::OnAnimatorsSetEnabled);
@@ -309,16 +268,30 @@ void UPropertyAnimatorCoreComponent::DestroyComponent(bool bPromoteChildren)
 	Super::DestroyComponent(bPromoteChildren);
 }
 
+void UPropertyAnimatorCoreComponent::TickComponent(float InDeltaTime, ELevelTick InTickType, FActorComponentTickFunction* InTickFunction)
+{
+	Super::TickComponent(InDeltaTime, InTickType, InTickFunction);
+
+	if (!EvaluateAnimators())
+	{
+		SetComponentTickEnabled(false);
+	}
+}
+
 void UPropertyAnimatorCoreComponent::PostLoad()
 {
 	Super::PostLoad();
 
+	PRAGMA_DISABLE_DEPRECATION_WARNINGS
+
 	// Migrate animators to new array property
 	if (!Animators.IsEmpty() && PropertyAnimators.IsEmpty())
 	{
 		PropertyAnimators = Animators.Array();
 		Animators.Empty();
 	}
+
+	PRAGMA_ENABLE_DEPRECATION_WARNINGS
 }
 
 #if WITH_EDITOR
@@ -379,3 +352,38 @@ void UPropertyAnimatorCoreComponent::ForEachAnimator(const TFunctionRef<bool(UPr
 		}
 	}
 }
+
+bool UPropertyAnimatorCoreComponent::EvaluateAnimators()
+{
+	if (!ShouldAnimate())
+	{
+		return false;
+	}
+
+	const UWorld* World = GetWorld();
+	const bool bIsSupportedWorld = IsValid(World) && (World->IsGameWorld() || World->IsEditorWorld());
+
+	if (!bIsSupportedWorld)
+	{
+		return false;
+	}
+
+	FInstancedPropertyBag Parameters;
+
+	for (const TObjectPtr<UPropertyAnimatorCoreBase>& Animator : PropertyAnimators)
+	{
+		if (!IsValid(Animator) || !Animator->GetAnimatorEnabled())
+		{
+			continue;
+		}
+
+		// Reset in case animator change values to avoid affecting following animators
+		Parameters.Reset();
+		Parameters.AddProperty(UPropertyAnimatorCoreBase::MagnitudeParameterName, EPropertyBagPropertyType::Float);
+		Parameters.SetValueFloat(UPropertyAnimatorCoreBase::MagnitudeParameterName, AnimatorsMagnitude);
+
+		Animator->EvaluateAnimator(Parameters);
+	}
+
+	return true;
+}
diff --git a/Engine/Plugins/Experimental/PropertyAnimatorCore/Source/PropertyAnimatorCore/Private/Properties/PropertyAnimatorCoreContext.cpp b/Engine/Plugins/Experimental/PropertyAnimatorCore/Source/PropertyAnimatorCore/Private/Properties/PropertyAnimatorCoreContext.cpp
--- a/Engine/Plugins/Experimental/PropertyAnimatorCore/Source/PropertyAnimatorCore/Private/Properties/PropertyAnimatorCoreContext.cpp
+++ b/Engine/Plugins/Experimental/PropertyAnimatorCore/Source/PropertyAnimatorCore/Private/Properties/PropertyAnimatorCoreContext.cpp
@@ -10,6 +10,8 @@
 #include "Properties/Handlers/PropertyAnimatorCoreHandlerBase.h"
 #include "Subsystems/PropertyAnimatorCoreSubsystem.h"
 
+DEFINE_LOG_CATEGORY_STATIC(LogPropertyAnimatorCoreContext, Log, All);
+
 TArray<FPropertyAnimatorCoreData> UPropertyAnimatorCoreContext::ResolveProperty(bool bInForEvaluation) const
 {
 	TArray<FPropertyAnimatorCoreData> ResolvedProperties;
@@ -227,6 +229,9 @@ bool UPropertyAnimatorCoreContext::ResolvePropertyOwner(AActor* InNewOwner)
 	bool bFound = IsValid(NewOwner);
 	const TArray<UObject*> OtherOuters = AnimatedProperty.GetOuters(AnimatedProperty.GetOwningActor());
 
+	const FProperty* MemberProperty = AnimatedProperty.GetMemberProperty();
+	UClass* PropertyOwningClass = MemberProperty->GetOwnerClass();
+
 	if (!OtherOuters.IsEmpty())
 	{
 		// Resolve using outers
@@ -236,6 +241,7 @@ bool UPropertyAnimatorCoreContext::ResolvePropertyOwner(AActor* InNewOwner)
 			TArray<UObject*> ThisOwnedObjects;
 			GetObjectsWithOuter(NewOwner, ThisOwnedObjects, false);
 
+			// Search for the same class and same name
 			for (UObject* ThisOuter : ThisOwnedObjects)
 			{
 				if (ThisOuter->GetClass() == OtherOuter->GetClass()
@@ -246,6 +252,27 @@ bool UPropertyAnimatorCoreContext::ResolvePropertyOwner(AActor* InNewOwner)
 					break;
 				}
 			}
+
+			// Search for the property owning class
+			if (!bFound)
+			{
+				for (UObject* ThisOuter : ThisOwnedObjects)
+				{
+					if (ThisOuter->GetClass() == OtherOuter->GetClass()
+						|| ThisOuter->IsA(PropertyOwningClass))
+					{
+						bFound = true;
+						NewOwner = ThisOuter;
+						break;
+					}
+				}
+
+				// Nothing found, stop searching
+				if (!bFound)
+				{
+					break;
+				}
+			}
 		}
 	}
 	else
@@ -261,28 +288,33 @@ bool UPropertyAnimatorCoreContext::ResolvePropertyOwner(AActor* InNewOwner)
 
 			for (UObject* ThisOuter : ThisOwnedObjects)
 			{
-				if (ThisOuter
-					&& ThisOuter->GetName().StartsWith(PathSegment))
+				if (ThisOuter && ThisOuter->GetName().StartsWith(PathSegment))
 				{
 					bFound = true;
 					NewOwner = ThisOuter;
 					break;
 				}
 			}
+
+			// Nothing found, stop searching
+			if (!bFound)
+			{
+				break;
+			}
 		}
 	}
 
-	const UObject* PreviousPropertyOwner = AnimatedProperty.GetOwner();
-
 	if (bFound
 		&& IsValid(NewOwner)
-		&& (!PreviousPropertyOwner || NewOwner->GetClass() == PreviousPropertyOwner->GetClass())
+		&& NewOwner->GetClass()->IsChildOf(PropertyOwningClass)
 		&& FindFProperty<FProperty>(NewOwner->GetClass(), AnimatedProperty.GetMemberPropertyName()))
 	{
 		SetAnimatedPropertyOwner(NewOwner);
 		return true;
 	}
 
+	UE_LOG(LogPropertyAnimatorCoreContext, Warning, TEXT("Could not resolve property owner %s on %s"), *AnimatedProperty.GetPathHash(), NewOwner ? *NewOwner->GetName() : TEXT("Invalid"))
+
 	return false;
 }
 
diff --git a/Engine/Plugins/Experimental/PropertyAnimatorCore/Source/PropertyAnimatorCore/Private/Subsystems/PropertyAnimatorCoreSubsystem.cpp b/Engine/Plugins/Experimental/PropertyAnimatorCore/Source/PropertyAnimatorCore/Private/Subsystems/PropertyAnimatorCoreSubsystem.cpp
--- a/Engine/Plugins/Experimental/PropertyAnimatorCore/Source/PropertyAnimatorCore/Private/Subsystems/PropertyAnimatorCoreSubsystem.cpp
+++ b/Engine/Plugins/Experimental/PropertyAnimatorCore/Source/PropertyAnimatorCore/Private/Subsystems/PropertyAnimatorCoreSubsystem.cpp
@@ -66,25 +66,25 @@ UPropertyAnimatorCoreSubsystem* UPropertyAnimatorCoreSubsystem::Get()
 	return nullptr;
 }
 
-bool UPropertyAnimatorCoreSubsystem::RegisterAnimatorClass(const UClass* InPropertyControllerClass)
+bool UPropertyAnimatorCoreSubsystem::RegisterAnimatorClass(const UClass* InAnimatorClass)
 {
-	if (!IsValid(InPropertyControllerClass))
+	if (!IsValid(InAnimatorClass))
 	{
 		return false;
 	}
 
-	if (!InPropertyControllerClass->IsChildOf(UPropertyAnimatorCoreBase::StaticClass())
-		|| InPropertyControllerClass->HasAnyClassFlags(CLASS_Abstract | CLASS_Deprecated | CLASS_NewerVersionExists))
+	if (!InAnimatorClass->IsChildOf(UPropertyAnimatorCoreBase::StaticClass())
+		|| InAnimatorClass->HasAnyClassFlags(CLASS_Abstract | CLASS_Deprecated | CLASS_NewerVersionExists))
 	{
 		return false;
 	}
 
-	if (IsAnimatorClassRegistered(InPropertyControllerClass))
+	if (IsAnimatorClassRegistered(InAnimatorClass))
 	{
 		return false;
 	}
 
-	if (UPropertyAnimatorCoreBase* CDO = InPropertyControllerClass->GetDefaultObject<UPropertyAnimatorCoreBase>())
+	if (UPropertyAnimatorCoreBase* CDO = InAnimatorClass->GetDefaultObject<UPropertyAnimatorCoreBase>())
 	{
 		AnimatorsWeak.Add(CDO);
 
@@ -94,16 +94,16 @@ bool UPropertyAnimatorCoreSubsystem::RegisterAnimatorClass(const UClass* InPrope
 	return false;
 }
 
-bool UPropertyAnimatorCoreSubsystem::UnregisterAnimatorClass(const UClass* InPropertyControllerClass)
+bool UPropertyAnimatorCoreSubsystem::UnregisterAnimatorClass(const UClass* InAnimatorClass)
 {
-	if (!IsValid(InPropertyControllerClass))
+	if (!IsValid(InAnimatorClass))
 	{
 		return false;
 	}
 
 	for (TSet<TWeakObjectPtr<UPropertyAnimatorCoreBase>>::TIterator It(AnimatorsWeak); It; ++It)
 	{
-		if ((*It)->GetClass() == InPropertyControllerClass)
+		if ((*It)->GetClass() == InAnimatorClass)
 		{
 			It.RemoveCurrent();
 			return true;
@@ -113,16 +113,16 @@ bool UPropertyAnimatorCoreSubsystem::UnregisterAnimatorClass(const UClass* InPro
 	return false;
 }
 
-bool UPropertyAnimatorCoreSubsystem::IsAnimatorClassRegistered(const UClass* InPropertyControllerClass) const
+bool UPropertyAnimatorCoreSubsystem::IsAnimatorClassRegistered(const UClass* InAnimatorClass) const
 {
-	if (!IsValid(InPropertyControllerClass))
+	if (!IsValid(InAnimatorClass))
 	{
 		return false;
 	}
 
-	for (const TWeakObjectPtr<UPropertyAnimatorCoreBase>& Controller : AnimatorsWeak)
+	for (const TWeakObjectPtr<UPropertyAnimatorCoreBase>& Animator : AnimatorsWeak)
 	{
-		if (Controller->GetClass() == InPropertyControllerClass)
+		if (Animator->GetClass() == InAnimatorClass)
 		{
 			return true;
 		}
@@ -187,9 +187,9 @@ bool UPropertyAnimatorCoreSubsystem::IsPropertySupported(const FPropertyAnimator
 
 TSet<UPropertyAnimatorCoreBase*> UPropertyAnimatorCoreSubsystem::GetPropertyLinkedAnimators(const FPropertyAnimatorCoreData& InPropertyData) const
 {
-	TSet<UPropertyAnimatorCoreBase*> ExistingControllers = GetExistingAnimators(InPropertyData);
+	TSet<UPropertyAnimatorCoreBase*> ExistingAnimators = GetExistingAnimators(InPropertyData);
 
-	for (TSet<UPropertyAnimatorCoreBase*>::TIterator It(ExistingControllers); It; ++It)
+	for (TSet<UPropertyAnimatorCoreBase*>::TIterator It(ExistingAnimators); It; ++It)
 	{
 		if (!(*It)->IsPropertyLinked(InPropertyData))
 		{
@@ -197,69 +197,69 @@ TSet<UPropertyAnimatorCoreBase*> UPropertyAnimatorCoreSubsystem::GetPropertyLink
 		}
 	}
 
-	return ExistingControllers;
+	return ExistingAnimators;
 }
 
 TSet<UPropertyAnimatorCoreBase*> UPropertyAnimatorCoreSubsystem::GetExistingAnimators(const FPropertyAnimatorCoreData& InPropertyData) const
 {
-	TSet<UPropertyAnimatorCoreBase*> ExistingPropertyControllers;
+	TSet<UPropertyAnimatorCoreBase*> ExistingAnimators;
 
 	if (!InPropertyData.IsResolved())
 	{
-		return ExistingPropertyControllers;
+		return ExistingAnimators;
 	}
 
 	const AActor* Actor = InPropertyData.GetOwningActor();
 
-	for (UPropertyAnimatorCoreBase* Controller : GetExistingAnimators(Actor))
+	for (UPropertyAnimatorCoreBase* Animator : GetExistingAnimators(Actor))
 	{
 		TSet<FPropertyAnimatorCoreData> OutProperties;
-		if (Controller->GetPropertiesSupported(InPropertyData, OutProperties))
+		if (Animator->GetPropertiesSupported(InPropertyData, OutProperties, /** SearchDepth */3))
 		{
-			ExistingPropertyControllers.Add(Controller);
+			ExistingAnimators.Add(Animator);
 		}
 	}
 
-	return ExistingPropertyControllers;
+	return ExistingAnimators;
 }
 
 TSet<UPropertyAnimatorCoreBase*> UPropertyAnimatorCoreSubsystem::GetExistingAnimators(const AActor* InActor) const
 {
-	TSet<UPropertyAnimatorCoreBase*> ExistingPropertyControllers;
+	TSet<UPropertyAnimatorCoreBase*> ExistingAnimators;
 
 	if (!IsValid(InActor))
 	{
-		return ExistingPropertyControllers;
+		return ExistingAnimators;
 	}
 
 	if (const UPropertyAnimatorCoreComponent* PropertyComponent = InActor->FindComponentByClass<UPropertyAnimatorCoreComponent>())
 	{
-		PropertyComponent->ForEachAnimator([&ExistingPropertyControllers](UPropertyAnimatorCoreBase* InController)->bool
+		PropertyComponent->ForEachAnimator([&ExistingAnimators](UPropertyAnimatorCoreBase* InAnimator)->bool
 		{
-			ExistingPropertyControllers.Add(InController);
+			ExistingAnimators.Add(InAnimator);
 			return true;
 		});
 	}
 
-	return ExistingPropertyControllers;
+	return ExistingAnimators;
 }
 
 TSet<UPropertyAnimatorCoreBase*> UPropertyAnimatorCoreSubsystem::GetAvailableAnimators(const FPropertyAnimatorCoreData* InPropertyData) const
 {
-	TSet<UPropertyAnimatorCoreBase*> AvailablePropertyControllers;
+	TSet<UPropertyAnimatorCoreBase*> AvailableAnimators;
 
 	if (InPropertyData && !InPropertyData->IsResolved())
 	{
-		return AvailablePropertyControllers;
+		return AvailableAnimators;
 	}
 
-	for (const TWeakObjectPtr<UPropertyAnimatorCoreBase>& Controller : AnimatorsWeak)
+	for (const TWeakObjectPtr<UPropertyAnimatorCoreBase>& Animator : AnimatorsWeak)
 	{
 		bool bIsPropertySupported = true;
 		if (InPropertyData)
 		{
 			TSet<FPropertyAnimatorCoreData> OutProperties;
-			if (!Controller->GetPropertiesSupported(*InPropertyData, OutProperties))
+			if (!Animator->GetPropertiesSupported(*InPropertyData, OutProperties, /** SearchDepth */3))
 			{
 				bIsPropertySupported = false;
 			}
@@ -267,23 +267,23 @@ TSet<UPropertyAnimatorCoreBase*> UPropertyAnimatorCoreSubsystem::GetAvailableAni
 
 		if (bIsPropertySupported)
 		{
-			AvailablePropertyControllers.Add(Controller.Get());
+			AvailableAnimators.Add(Animator.Get());
 		}
 	}
 
-	return AvailablePropertyControllers;
+	return AvailableAnimators;
 }
 
 TSet<UPropertyAnimatorCoreBase*> UPropertyAnimatorCoreSubsystem::GetAvailableAnimators() const
 {
-	TSet<UPropertyAnimatorCoreBase*> AvailablePropertyControllers;
+	TSet<UPropertyAnimatorCoreBase*> AvailableAnimators;
 
-	for (const TWeakObjectPtr<UPropertyAnimatorCoreBase>& Controller : AnimatorsWeak)
+	for (const TWeakObjectPtr<UPropertyAnimatorCoreBase>& Animator : AnimatorsWeak)
 	{
-		AvailablePropertyControllers.Add(Controller.Get());
+		AvailableAnimators.Add(Animator.Get());
 	}
 
-	return AvailablePropertyControllers;
+	return AvailableAnimators;
 }
 
 UPropertyAnimatorCoreBase* UPropertyAnimatorCoreSubsystem::CreateAnimator(AActor* InActor, const UClass* InAnimatorClass, UPropertyAnimatorCorePresetBase* InPreset, bool bInTransact) const
diff --git a/Engine/Plugins/Experimental/PropertyAnimatorCore/Source/PropertyAnimatorCore/Public/Components/PropertyAnimatorCoreComponent.h b/Engine/Plugins/Experimental/PropertyAnimatorCore/Source/PropertyAnimatorCore/Public/Components/PropertyAnimatorCoreComponent.h
--- a/Engine/Plugins/Experimental/PropertyAnimatorCore/Source/PropertyAnimatorCore/Public/Components/PropertyAnimatorCoreComponent.h
+++ b/Engine/Plugins/Experimental/PropertyAnimatorCore/Source/PropertyAnimatorCore/Public/Components/PropertyAnimatorCoreComponent.h
@@ -49,12 +49,16 @@ public:
 	/** Process a function for each controller, stops when false is returned otherwise continue until the end */
 	PROPERTYANIMATORCORE_API void ForEachAnimator(TFunctionRef<bool(UPropertyAnimatorCoreBase*)> InFunction) const;
 
+	/** Checks if this component animators should be active */
+	bool ShouldAnimate() const;
+
 protected:
 	static FName GetAnimatorName(const UPropertyAnimatorCoreBase* InAnimator);
 
 	//~ Begin UActorComponent
 	virtual void OnComponentCreated() override;
 	virtual void DestroyComponent(bool bPromoteChildren) override;
+	virtual void TickComponent(float InDeltaTime, ELevelTick InTickType, FActorComponentTickFunction* InTickFunction) override;
 	//~ End UActorComponent
 
 	//~ Begin UObject
@@ -86,12 +90,14 @@ protected:
 	/** Change global state for animators */
 	void OnAnimatorsSetEnabled(const UWorld* InWorld, bool bInEnabled, bool bInTransact);
 
+	/** Callback when PropertyAnimators changed */
 	void OnAnimatorsChanged();
 
+	/** Callback when global enabled state is changed */
 	void OnAnimatorsEnabledChanged();
 
-	/** Checks if this component animators should tick */
-	bool ShouldAnimatorsTick() const;
+	/** Evaluate only specified animators */
+	bool EvaluateAnimators();
 
 	/** Animators linked to this actor, they contain only properties within this actor */
 	UPROPERTY(EditAnywhere, BlueprintReadWrite, NoClear, Export, Instanced, Setter="SetAnimators", Category="Animator", meta=(TitleProperty="AnimatorDisplayName"))
@@ -106,9 +112,8 @@ protected:
 	float AnimatorsMagnitude = 1.f;
 
 private:
-	virtual void TickComponent(float InDeltaTime, ELevelTick InTickType, FActorComponentTickFunction* InThisTickFunction) override;
-
 	/** Deprecated property set, will be migrated to PropertyAnimators property on load */
+	UE_DEPRECATED(5.5, "Moved to PropertyAnimators")
 	UPROPERTY()
 	TSet<TObjectPtr<UPropertyAnimatorCoreBase>> Animators;
 
diff --git a/Engine/Plugins/Experimental/PropertyAnimatorCore/Source/PropertyAnimatorCore/Public/Subsystems/PropertyAnimatorCoreSubsystem.h b/Engine/Plugins/Experimental/PropertyAnimatorCore/Source/PropertyAnimatorCore/Public/Subsystems/PropertyAnimatorCoreSubsystem.h
--- a/Engine/Plugins/Experimental/PropertyAnimatorCore/Source/PropertyAnimatorCore/Public/Subsystems/PropertyAnimatorCoreSubsystem.h
+++ b/Engine/Plugins/Experimental/PropertyAnimatorCore/Source/PropertyAnimatorCore/Public/Subsystems/PropertyAnimatorCoreSubsystem.h
@@ -34,29 +34,29 @@ public:
 	virtual void Deinitialize() override;
 	//~ End USubsystem
 
-	/** Register the property controller class to allow its usage */
-	PROPERTYANIMATORCORE_API bool RegisterAnimatorClass(const UClass* InPropertyControllerClass);
+	/** Register the property animator class to allow its usage */
+	PROPERTYANIMATORCORE_API bool RegisterAnimatorClass(const UClass* InAnimatorClass);
 
-	/** Unregister the property controller class to disallow its usage */
-	PROPERTYANIMATORCORE_API bool UnregisterAnimatorClass(const UClass* InPropertyControllerClass);
+	/** Unregister the property animator class to disallow its usage */
+	PROPERTYANIMATORCORE_API bool UnregisterAnimatorClass(const UClass* InAnimatorClass);
 
-	/** Checks if the property controller class is already registered */
-	PROPERTYANIMATORCORE_API bool IsAnimatorClassRegistered(const UClass* InPropertyControllerClass) const;
+	/** Checks if the property animator class is already registered */
+	PROPERTYANIMATORCORE_API bool IsAnimatorClassRegistered(const UClass* InAnimatorClass) const;
 
 	/** Gets the animator CDO registered from the class */
 	UPropertyAnimatorCoreBase* GetAnimatorRegistered(const UClass* InAnimatorClass) const;
 
-	/** Returns true if any controller is able to control that property or nested otherwise false */
+	/** Returns true if any animator is able to control that property or nested otherwise false */
 	PROPERTYANIMATORCORE_API bool IsPropertySupported(const FPropertyAnimatorCoreData& InPropertyData, bool bInCheckNestedProperties = true) const;
 
 	/** Find all animators linked to the property */
 	PROPERTYANIMATORCORE_API TSet<UPropertyAnimatorCoreBase*> GetPropertyLinkedAnimators(const FPropertyAnimatorCoreData& InPropertyData) const;
 
-	/** Returns a set of existing property controller objects in owner that supports that property */
+	/** Returns a set of existing property animator objects in owner that supports that property */
 	PROPERTYANIMATORCORE_API TSet<UPropertyAnimatorCoreBase*> GetExistingAnimators(const FPropertyAnimatorCoreData& InPropertyData) const;
 	PROPERTYANIMATORCORE_API TSet<UPropertyAnimatorCoreBase*> GetExistingAnimators(const AActor* InActor) const;
 
-	/** Returns a set of property controller CDO that supports that property */
+	/** Returns a set of property animator CDO that supports that property */
 	PROPERTYANIMATORCORE_API TSet<UPropertyAnimatorCoreBase*> GetAvailableAnimators(const FPropertyAnimatorCoreData* InPropertyData) const;
 	PROPERTYANIMATORCORE_API TSet<UPropertyAnimatorCoreBase*> GetAvailableAnimators() const;
 
