MotionDesign : PropertyAnimator
- Setting a target framerate on time sources to change smoothness of animation

diff --git a/Engine/Plugins/Experimental/PropertyAnimatorCore/Source/PropertyAnimatorCore/Private/TimeSources/PropertyAnimatorCoreTimeSourceBase.cpp b/Engine/Plugins/Experimental/PropertyAnimatorCore/Source/PropertyAnimatorCore/Private/TimeSources/PropertyAnimatorCoreTimeSourceBase.cpp
--- a/Engine/Plugins/Experimental/PropertyAnimatorCore/Source/PropertyAnimatorCore/Private/TimeSources/PropertyAnimatorCoreTimeSourceBase.cpp
+++ b/Engine/Plugins/Experimental/PropertyAnimatorCore/Source/PropertyAnimatorCore/Private/TimeSources/PropertyAnimatorCoreTimeSourceBase.cpp
@@ -4,13 +4,6 @@
 
 #include "Animators/PropertyAnimatorCoreBase.h"
 
-#if WITH_EDITOR
-FName UPropertyAnimatorCoreTimeSourceBase::GetTimeElapsedPropertyName()
-{
-	return GET_MEMBER_NAME_CHECKED(UPropertyAnimatorCoreTimeSourceBase, TimeElapsed);
-}
-#endif
-
 void UPropertyAnimatorCoreTimeSourceBase::ActivateTimeSource()
 {
 	if (IsTimeSourceActive())
@@ -40,17 +33,34 @@ TOptional<double> UPropertyAnimatorCoreTimeSourceBase::GetConditionalTimeElapsed
 		return TOptional<double>();
 	}
 
-	TimeElapsed = GetTimeElapsed();
+	double NewTimeElapsed = GetTimeElapsed();
 
-	if (!IsValidTimeElapsed(TimeElapsed))
+	if (!IsValidTimeElapsed(NewTimeElapsed))
 	{
 		return TOptional<double>();
 	}
 
-	return TimeElapsed;
+	LastTimeElapsed = NewTimeElapsed;
+
+	return LastTimeElapsed;
 }
 
 UPropertyAnimatorCoreBase* UPropertyAnimatorCoreTimeSourceBase::GetAnimator() const
 {
 	return GetTypedOuter<UPropertyAnimatorCoreBase>();
 }
+
+void UPropertyAnimatorCoreTimeSourceBase::SetFrameRate(float InFrameRate)
+{
+	FrameRate = FMath::Max(UE_KINDA_SMALL_NUMBER, InFrameRate);
+}
+
+void UPropertyAnimatorCoreTimeSourceBase::SetUseFrameRate(bool bInUseFrameRate)
+{
+	bUseFrameRate = bInUseFrameRate;
+}
+
+bool UPropertyAnimatorCoreTimeSourceBase::IsValidTimeElapsed(double InTimeElapsed) const
+{
+	return !bUseFrameRate || FMath::IsNearlyZero(FrameRate) || FMath::Abs(InTimeElapsed - LastTimeElapsed) > FMath::Abs(1.f / FrameRate);
+}
diff --git a/Engine/Plugins/Experimental/PropertyAnimatorCore/Source/PropertyAnimatorCore/Public/TimeSources/PropertyAnimatorCoreTimeSourceBase.h b/Engine/Plugins/Experimental/PropertyAnimatorCore/Source/PropertyAnimatorCore/Public/TimeSources/PropertyAnimatorCoreTimeSourceBase.h
--- a/Engine/Plugins/Experimental/PropertyAnimatorCore/Source/PropertyAnimatorCore/Public/TimeSources/PropertyAnimatorCoreTimeSourceBase.h
+++ b/Engine/Plugins/Experimental/PropertyAnimatorCore/Source/PropertyAnimatorCore/Public/TimeSources/PropertyAnimatorCoreTimeSourceBase.h
@@ -18,10 +18,6 @@ class UPropertyAnimatorCoreTimeSourceBase : public UObject
 	friend class UPropertyAnimatorCoreSubsystem;
 
 public:
-#if WITH_EDITOR
-	PROPERTYANIMATORCORE_API static FName GetTimeElapsedPropertyName();
-#endif
-
 	UPropertyAnimatorCoreTimeSourceBase()
 		: UPropertyAnimatorCoreTimeSourceBase(NAME_None)
 	{}
@@ -48,6 +44,18 @@ public:
 		return TimeSourceName;
 	}
 
+	void SetFrameRate(float InFrameRate);
+	float GetFrameRate() const
+	{
+		return FrameRate;
+	}
+
+	void SetUseFrameRate(bool bInUseFrameRate);
+	bool GetUseFrameRate() const
+	{
+		return bUseFrameRate;
+	}
+
 protected:
 	/** Returns the time elapsed for animators */
 	virtual double GetTimeElapsed()
@@ -62,10 +70,7 @@ protected:
 	}
 
 	/** Check if the time elapsed is valid based on the context */
-	virtual bool IsValidTimeElapsed(double InTimeElapsed) const
-	{
-		return true;
-	}
+	PROPERTYANIMATORCORE_API virtual bool IsValidTimeElapsed(double InTimeElapsed) const;
 
 	/** Time source CDO is registered by subsystem */
 	virtual void OnTimeSourceRegistered() {}
@@ -80,13 +85,21 @@ protected:
 	virtual void OnTimeSourceInactive() {}
 
 private:
-	/** Only used to display time elapsed */
-	UPROPERTY(Transient, VisibleInstanceOnly, Category="Animator", meta=(NoResetToDefault))
-	double TimeElapsed = 0;
+	/** Use a specific framerate */
+	UPROPERTY(EditInstanceOnly, Setter="SetUseFrameRate", Getter="GetUseFrameRate", Category="Animator", meta=(InlineEditConditionToggle))
+	bool bUseFrameRate = false;
+
+	/** The frame rate to target for the animator effect */
+	UPROPERTY(EditInstanceOnly, Setter, Getter, Category="Animator", meta=(ClampMin="0", EditCondition="bUseFrameRate"))
+	float FrameRate = 30.f;
 
 	/** Name used to display this time source to the user */
 	UPROPERTY(Transient)
 	FName TimeSourceName;
 
+	/** Cached time elapsed */
+	double LastTimeElapsed = 0;
+
+	/** Is this time source active on the animator */
 	bool bTimeSourceActive = false;
 };
\ No newline at end of file
diff --git a/Engine/Plugins/Experimental/PropertyAnimatorCore/Source/PropertyAnimatorCoreEditor/Private/Customizations/PropertyAnimatorCoreEditorStackCustomization.cpp b/Engine/Plugins/Experimental/PropertyAnimatorCore/Source/PropertyAnimatorCoreEditor/Private/Customizations/PropertyAnimatorCoreEditorStackCustomization.cpp
--- a/Engine/Plugins/Experimental/PropertyAnimatorCore/Source/PropertyAnimatorCoreEditor/Private/Customizations/PropertyAnimatorCoreEditorStackCustomization.cpp
+++ b/Engine/Plugins/Experimental/PropertyAnimatorCore/Source/PropertyAnimatorCoreEditor/Private/Customizations/PropertyAnimatorCoreEditorStackCustomization.cpp
@@ -294,13 +294,11 @@ void UPropertyAnimatorCoreEditorStackCustomization::CustomizeItemBody(const FOpe
 	else if (InItem->IsA<UPropertyAnimatorCoreBase>())
 	{
 		FBoolProperty* EnableProperty = FindFProperty<FBoolProperty>(UPropertyAnimatorCoreBase::StaticClass(), GET_MEMBER_NAME_CHECKED(UPropertyAnimatorCoreBase, bAnimatorEnabled));
-		FProperty* TimeElapsedProperty = FindFProperty<FProperty>(UPropertyAnimatorCoreTimeSourceBase::StaticClass(), UPropertyAnimatorCoreTimeSourceBase::GetTimeElapsedPropertyName());
 		FProperty* LinkedPropertiesProperty = FindFProperty<FProperty>(UPropertyAnimatorCoreBase::StaticClass(), UPropertyAnimatorCoreBase::GetLinkedPropertiesPropertyName());
 
 		InBodyBuilder
 			.SetShowDetailsView(true)
 			.DisallowProperty(EnableProperty)
-			.DisallowProperty(TimeElapsedProperty)
 			.ExpandProperty(LinkedPropertiesProperty);
 	}
 
