MotionDesign : PropertyAnimator
- Add animator track action button from the sequencer time source

diff --git a/Engine/Plugins/Experimental/PropertyAnimatorCore/Source/PropertyAnimatorCore/Private/Sequencer/MovieSceneAnimatorSection.cpp b/Engine/Plugins/Experimental/PropertyAnimatorCore/Source/PropertyAnimatorCore/Private/Sequencer/MovieSceneAnimatorSection.cpp
--- a/Engine/Plugins/Experimental/PropertyAnimatorCore/Source/PropertyAnimatorCore/Private/Sequencer/MovieSceneAnimatorSection.cpp
+++ b/Engine/Plugins/Experimental/PropertyAnimatorCore/Source/PropertyAnimatorCore/Private/Sequencer/MovieSceneAnimatorSection.cpp
@@ -7,3 +7,8 @@ UMovieSceneAnimatorSection::UMovieSceneAnimatorSection()
 {
 	bSupportsInfiniteRange = true;
 }
+
+void UMovieSceneAnimatorSection::SetChannel(uint8 InChannel)
+{
+	Channel = InChannel;
+}
diff --git a/Engine/Plugins/Experimental/PropertyAnimatorCore/Source/PropertyAnimatorCore/Private/Sequencer/MovieSceneAnimatorSection.h b/Engine/Plugins/Experimental/PropertyAnimatorCore/Source/PropertyAnimatorCore/Private/Sequencer/MovieSceneAnimatorSection.h
deleted file mode 100644
diff --git a/Engine/Plugins/Experimental/PropertyAnimatorCore/Source/PropertyAnimatorCore/Private/Sequencer/MovieSceneAnimatorTrack.cpp b/Engine/Plugins/Experimental/PropertyAnimatorCore/Source/PropertyAnimatorCore/Private/Sequencer/MovieSceneAnimatorTrack.cpp
--- a/Engine/Plugins/Experimental/PropertyAnimatorCore/Source/PropertyAnimatorCore/Private/Sequencer/MovieSceneAnimatorTrack.cpp
+++ b/Engine/Plugins/Experimental/PropertyAnimatorCore/Source/PropertyAnimatorCore/Private/Sequencer/MovieSceneAnimatorTrack.cpp
@@ -16,6 +16,24 @@ UMovieSceneAnimatorTrack::UMovieSceneAnimatorTrack()
 	SupportedBlendTypes = FMovieSceneBlendTypeField::None();
 }
 
+int32 UMovieSceneAnimatorTrack::GetChannelCount(uint8 InChannel) const
+{
+	int32 Count = 0;
+
+	for (UMovieSceneSection* Section : GetAllSections())
+	{
+		if (const UMovieSceneAnimatorSection* AnimatorSection = Cast<UMovieSceneAnimatorSection>(Section))
+		{
+			if (AnimatorSection->GetChannel() == InChannel)
+			{
+				Count++;
+			}
+		}
+	}
+
+	return Count;
+}
+
 bool UMovieSceneAnimatorTrack::SupportsType(TSubclassOf<UMovieSceneSection> InSectionClass) const
 {
 	return InSectionClass == UMovieSceneAnimatorSection::StaticClass();
diff --git a/Engine/Plugins/Experimental/PropertyAnimatorCore/Source/PropertyAnimatorCore/Private/TimeSources/PropertyAnimatorCoreSequencerTimeSource.cpp b/Engine/Plugins/Experimental/PropertyAnimatorCore/Source/PropertyAnimatorCore/Private/TimeSources/PropertyAnimatorCoreSequencerTimeSource.cpp
--- a/Engine/Plugins/Experimental/PropertyAnimatorCore/Source/PropertyAnimatorCore/Private/TimeSources/PropertyAnimatorCoreSequencerTimeSource.cpp
+++ b/Engine/Plugins/Experimental/PropertyAnimatorCore/Source/PropertyAnimatorCore/Private/TimeSources/PropertyAnimatorCoreSequencerTimeSource.cpp
@@ -30,12 +30,12 @@ void UPropertyAnimatorCoreSequencerTimeSource::OnTimeSourceInactive()
 
 void UPropertyAnimatorCoreSequencerTimeSource::SetChannel(uint8 InChannel)
 {
-	Channel = InChannel;
+	ChannelData.Channel = InChannel;
 }
 
 void UPropertyAnimatorCoreSequencerTimeSource::OnSequencerTimeEvaluated(uint8 InChannel, double InTimeEval)
 {
-	if (Channel == InChannel)
+	if (ChannelData.Channel == InChannel)
 	{
 		EvalTime = InTimeEval;
 	}
diff --git a/Engine/Plugins/Experimental/PropertyAnimatorCore/Source/PropertyAnimatorCore/Public/Sequencer/MovieSceneAnimatorSection.h b/Engine/Plugins/Experimental/PropertyAnimatorCore/Source/PropertyAnimatorCore/Public/Sequencer/MovieSceneAnimatorSection.h
new file mode 100644
--- /dev/null
+++ b/Engine/Plugins/Experimental/PropertyAnimatorCore/Source/PropertyAnimatorCore/Public/Sequencer/MovieSceneAnimatorSection.h
@@ -0,0 +1,27 @@
+// Copyright Epic Games, Inc. All Rights Reserved.
+
+#pragma once
+
+#include "MovieSceneSection.h"
+#include "MovieSceneAnimatorSection.generated.h"
+
+/** Movie scene section for a sequencer animator channel */
+UCLASS(MinimalAPI)
+class UMovieSceneAnimatorSection : public UMovieSceneSection
+{
+	GENERATED_BODY()
+
+public:
+	UMovieSceneAnimatorSection();
+
+	PROPERTYANIMATORCORE_API void SetChannel(uint8 InChannel);
+	uint8 GetChannel() const
+	{
+		return Channel;
+	}
+
+protected:
+	/** Channel used to push sequencer time to */
+	UPROPERTY(EditAnywhere, Category="Animator")
+	uint8 Channel = 0;
+};
diff --git a/Engine/Plugins/Experimental/PropertyAnimatorCore/Source/PropertyAnimatorCore/Public/Sequencer/MovieSceneAnimatorTrack.h b/Engine/Plugins/Experimental/PropertyAnimatorCore/Source/PropertyAnimatorCore/Public/Sequencer/MovieSceneAnimatorTrack.h
--- a/Engine/Plugins/Experimental/PropertyAnimatorCore/Source/PropertyAnimatorCore/Public/Sequencer/MovieSceneAnimatorTrack.h
+++ b/Engine/Plugins/Experimental/PropertyAnimatorCore/Source/PropertyAnimatorCore/Public/Sequencer/MovieSceneAnimatorTrack.h
@@ -15,6 +15,9 @@ class UMovieSceneAnimatorTrack : public UMovieSceneNameableTrack, public IMovieS
 public:
 	UMovieSceneAnimatorTrack();
 
+	/** Get the number of section using the channel under this track */
+	PROPERTYANIMATORCORE_API int32 GetChannelCount(uint8 InChannel) const;
+
 	//~ Begin UMovieSceneTrack
 	virtual bool SupportsType(TSubclassOf<UMovieSceneSection> InSectionClass) const override;
 	virtual UMovieSceneSection* CreateNewSection() override;
diff --git a/Engine/Plugins/Experimental/PropertyAnimatorCore/Source/PropertyAnimatorCore/Public/TimeSources/PropertyAnimatorCoreSequencerTimeSource.h b/Engine/Plugins/Experimental/PropertyAnimatorCore/Source/PropertyAnimatorCore/Public/TimeSources/PropertyAnimatorCoreSequencerTimeSource.h
--- a/Engine/Plugins/Experimental/PropertyAnimatorCore/Source/PropertyAnimatorCore/Public/TimeSources/PropertyAnimatorCoreSequencerTimeSource.h
+++ b/Engine/Plugins/Experimental/PropertyAnimatorCore/Source/PropertyAnimatorCore/Public/TimeSources/PropertyAnimatorCoreSequencerTimeSource.h
@@ -5,6 +5,15 @@
 #include "PropertyAnimatorCoreTimeSourceBase.h"
 #include "PropertyAnimatorCoreSequencerTimeSource.generated.h"
 
+USTRUCT()
+struct FPropertyAnimatorCoreSequencerTimeSourceChannel
+{
+	GENERATED_BODY()
+
+	UPROPERTY(EditInstanceOnly, Category="Animator")
+	uint8 Channel = 0;
+};
+
 /** Sequencer time source that sync with animator track channel */
 UCLASS(MinimalAPI)
 class UPropertyAnimatorCoreSequencerTimeSource : public UPropertyAnimatorCoreTimeSourceBase
@@ -29,15 +38,15 @@ public:
 	void SetChannel(uint8 InChannel);
 	uint8 GetChannel() const
 	{
-		return Channel;
+		return ChannelData.Channel;
 	}
 
 protected:
 	void OnSequencerTimeEvaluated(uint8 InChannel, double InTimeEval);
 
 	/** Channel to sample time from */
-	UPROPERTY(EditInstanceOnly, Setter, Getter, Category="Animator")
-	uint8 Channel = 0;
+	UPROPERTY(EditInstanceOnly, DisplayName="Channel", Category="Animator")
+	FPropertyAnimatorCoreSequencerTimeSourceChannel ChannelData;
 
 	/** Last evaluated time received */
 	TOptional<double> EvalTime;
diff --git a/Engine/Plugins/Experimental/PropertyAnimatorCore/Source/PropertyAnimatorCoreEditor/Private/Customizations/PropertyAnimatorCoreEditorSequencerTimeSourceChannelTypeCustomization.cpp b/Engine/Plugins/Experimental/PropertyAnimatorCore/Source/PropertyAnimatorCoreEditor/Private/Customizations/PropertyAnimatorCoreEditorSequencerTimeSourceChannelTypeCustomization.cpp
new file mode 100644
--- /dev/null
+++ b/Engine/Plugins/Experimental/PropertyAnimatorCore/Source/PropertyAnimatorCoreEditor/Private/Customizations/PropertyAnimatorCoreEditorSequencerTimeSourceChannelTypeCustomization.cpp
@@ -0,0 +1,128 @@
+// Copyright Epic Games, Inc. All Rights Reserved.
+
+#include "Customizations/PropertyAnimatorCoreEditorSequencerTimeSourceChannelTypeCustomization.h"
+
+#include "DetailWidgetRow.h"
+#include "Editor.h"
+#include "Selection.h"
+#include "Sequencer/MovieSceneAnimatorTrackEditor.h"
+#include "TimeSources/PropertyAnimatorCoreSequencerTimeSource.h"
+#include "Widgets/Input/SButton.h"
+#include "Widgets/SBoxPanel.h"
+#include "Widgets/Text/STextBlock.h"
+
+#define LOCTEXT_NAMESPACE "PropertyAnimatorCoreEditorSequencerTimeSourceChannelTypeCustomization"
+
+void FPropertyAnimatorCoreEditorSequencerTimeSourceChannelTypeCustomization::CustomizeHeader(TSharedRef<IPropertyHandle> InPropertyHandle, FDetailWidgetRow& InRow, IPropertyTypeCustomizationUtils& InUtils)
+{
+	if (!InPropertyHandle->IsValidHandle())
+	{
+		return;
+	}
+
+	ChannelPropertyHandle = InPropertyHandle->GetChildHandle(GET_MEMBER_NAME_CHECKED(FPropertyAnimatorCoreSequencerTimeSourceChannel, Channel));
+
+	if (!ChannelPropertyHandle)
+	{
+		return;
+	}
+
+	InRow.NameContent()
+	[
+		InPropertyHandle->CreatePropertyNameWidget()
+	];
+
+	InRow.ValueContent()
+	.HAlign(HAlign_Fill)
+	[
+		SNew(SHorizontalBox)
+		+ SHorizontalBox::Slot()
+		.FillWidth(1.f)
+		[
+			ChannelPropertyHandle->CreatePropertyValueWidget()
+		]
+		+ SHorizontalBox::Slot()
+		.Padding(5.f, 0.f)
+		.AutoWidth()
+		[
+			SNew(SButton)
+			.HAlign(HAlign_Fill)
+			.Visibility(this, &FPropertyAnimatorCoreEditorSequencerTimeSourceChannelTypeCustomization::GetCreateTrackButtonVisibility)
+			.OnClicked(this, &FPropertyAnimatorCoreEditorSequencerTimeSourceChannelTypeCustomization::OnCreateTrackButtonClicked)
+			.IsEnabled(this, &FPropertyAnimatorCoreEditorSequencerTimeSourceChannelTypeCustomization::IsCreateTrackButtonEnabled)
+			[
+				SNew(STextBlock)
+				.Font(IPropertyTypeCustomizationUtils::GetRegularFont())
+				.Text(LOCTEXT("AddSequencerTrack", "Add Sequencer track"))
+			]
+		]
+	];
+}
+
+void FPropertyAnimatorCoreEditorSequencerTimeSourceChannelTypeCustomization::CustomizeChildren(TSharedRef<IPropertyHandle> InPropertyHandle, IDetailChildrenBuilder& InBuilder, IPropertyTypeCustomizationUtils& InUtils)
+{
+}
+
+FReply FPropertyAnimatorCoreEditorSequencerTimeSourceChannelTypeCustomization::OnCreateTrackButtonClicked()
+{
+	TArray<AActor*> SelectedActors;
+
+	if (GEditor)
+	{
+		if (USelection* ActorSelection = GEditor->GetSelectedActors())
+		{
+			ActorSelection->GetSelectedObjects<AActor>(SelectedActors);
+		}
+	}
+
+	if (ChannelPropertyHandle.IsValid())
+	{
+		uint8 Channel = 0;
+
+		if (ChannelPropertyHandle->GetValue(Channel) == FPropertyAccess::Success)
+		{
+			FMovieSceneAnimatorTrackEditor::OnAddAnimatorTrack.Broadcast(Channel);
+
+			// Reselect actors after track was created
+			if (GEditor)
+			{
+				GEditor->SelectNone(/** Notify */false, /** DeselectBSP */true);
+
+				for (int32 Index = 0; Index < SelectedActors.Num(); Index++)
+				{
+					GEditor->SelectActor(SelectedActors[Index], /** Selected */true, /** Notify */Index == SelectedActors.Num() - 1);
+				}
+			}
+		}
+	}
+
+	return FReply::Handled();
+}
+
+EVisibility FPropertyAnimatorCoreEditorSequencerTimeSourceChannelTypeCustomization::GetCreateTrackButtonVisibility() const
+{
+	return ChannelPropertyHandle.IsValid()
+		&& ChannelPropertyHandle->GetNumPerObjectValues() == 1
+		? EVisibility::Visible
+		: EVisibility::Collapsed;
+}
+
+bool FPropertyAnimatorCoreEditorSequencerTimeSourceChannelTypeCustomization::IsCreateTrackButtonEnabled() const
+{
+	if (!ChannelPropertyHandle.IsValid() || ChannelPropertyHandle->GetNumPerObjectValues() != 1)
+	{
+		return false;
+	}
+
+	uint8 Channel = 0;
+	int32 Count = 0;
+
+	if (ChannelPropertyHandle->GetValue(Channel) == FPropertyAccess::Success)
+	{
+		FMovieSceneAnimatorTrackEditor::OnGetAnimatorTrackCount.Broadcast(Channel, Count);
+	}
+
+	return Count == 0;
+}
+
+#undef LOCTEXT_NAMESPACE
diff --git a/Engine/Plugins/Experimental/PropertyAnimatorCore/Source/PropertyAnimatorCoreEditor/Private/Customizations/PropertyAnimatorCoreEditorSequencerTimeSourceChannelTypeCustomization.h b/Engine/Plugins/Experimental/PropertyAnimatorCore/Source/PropertyAnimatorCoreEditor/Private/Customizations/PropertyAnimatorCoreEditorSequencerTimeSourceChannelTypeCustomization.h
new file mode 100644
--- /dev/null
+++ b/Engine/Plugins/Experimental/PropertyAnimatorCore/Source/PropertyAnimatorCoreEditor/Private/Customizations/PropertyAnimatorCoreEditorSequencerTimeSourceChannelTypeCustomization.h
@@ -0,0 +1,32 @@
+// Copyright Epic Games, Inc. All Rights Reserved.
+
+#pragma once
+
+#include "IPropertyTypeCustomization.h"
+#include "Templates/SharedPointer.h"
+
+class FReply;
+class IPropertyHandle;
+struct EVisibility;
+
+/** Type customization for FPropertyAnimatorCoreSequencerTimeSourceChannel */
+class FPropertyAnimatorCoreEditorSequencerTimeSourceChannelTypeCustomization : public IPropertyTypeCustomization
+{
+public:
+	static TSharedRef<IPropertyTypeCustomization> MakeInstance()
+	{
+		return MakeShared<FPropertyAnimatorCoreEditorSequencerTimeSourceChannelTypeCustomization>();
+	}
+
+	//~ Begin IPropertyTypeCustomization
+	virtual void CustomizeHeader(TSharedRef<IPropertyHandle> InPropertyHandle, FDetailWidgetRow& InRow, IPropertyTypeCustomizationUtils& InUtils) override;
+	virtual void CustomizeChildren(TSharedRef<IPropertyHandle> InPropertyHandle, IDetailChildrenBuilder& InBuilder, IPropertyTypeCustomizationUtils& InUtils) override;
+	//~ End IPropertyTypeCustomization
+
+protected:
+	FReply OnCreateTrackButtonClicked();
+	EVisibility GetCreateTrackButtonVisibility() const;
+	bool IsCreateTrackButtonEnabled() const;
+
+	TSharedPtr<IPropertyHandle> ChannelPropertyHandle;
+};
diff --git a/Engine/Plugins/Experimental/PropertyAnimatorCore/Source/PropertyAnimatorCoreEditor/Private/PropertyAnimatorCoreEditorModule.cpp b/Engine/Plugins/Experimental/PropertyAnimatorCore/Source/PropertyAnimatorCoreEditor/Private/PropertyAnimatorCoreEditorModule.cpp
--- a/Engine/Plugins/Experimental/PropertyAnimatorCore/Source/PropertyAnimatorCoreEditor/Private/PropertyAnimatorCoreEditorModule.cpp
+++ b/Engine/Plugins/Experimental/PropertyAnimatorCore/Source/PropertyAnimatorCoreEditor/Private/PropertyAnimatorCoreEditorModule.cpp
@@ -3,18 +3,21 @@
 #include "PropertyAnimatorCoreEditorModule.h"
 
 #include "Animators/PropertyAnimatorCoreBase.h"
-#include "Customizations/PropertyAnimatorCoreEditorDetailCustomization.h"
 #include "Customizations/PropertyAnimatorCoreEditorContextTypeCustomization.h"
+#include "Customizations/PropertyAnimatorCoreEditorDetailCustomization.h"
+#include "Customizations/PropertyAnimatorCoreEditorSequencerTimeSourceChannelTypeCustomization.h"
 #include "ISequencerModule.h"
 #include "Modules/ModuleManager.h"
 #include "Properties/PropertyAnimatorCoreContext.h"
 #include "PropertyEditorModule.h"
 #include "Sequencer/MovieSceneAnimatorTrackEditor.h"
+#include "TimeSources/PropertyAnimatorCoreSequencerTimeSource.h"
 
 void FPropertyAnimatorCoreEditorModule::StartupModule()
 {
 	FPropertyEditorModule& PropertyModule = FModuleManager::GetModuleChecked<FPropertyEditorModule>("PropertyEditor");
 	PropertyModule.RegisterCustomPropertyTypeLayout(UPropertyAnimatorCoreContext::StaticClass()->GetFName(), FOnGetPropertyTypeCustomizationInstance::CreateStatic(&FPropertyAnimatorCoreEditorContextTypeCustomization::MakeInstance));
+	PropertyModule.RegisterCustomPropertyTypeLayout(FPropertyAnimatorCoreSequencerTimeSourceChannel::StaticStruct()->GetFName(), FOnGetPropertyTypeCustomizationInstance::CreateStatic(&FPropertyAnimatorCoreEditorSequencerTimeSourceChannelTypeCustomization::MakeInstance));
 	PropertyModule.RegisterCustomClassLayout(UPropertyAnimatorCoreBase::StaticClass()->GetFName(), FOnGetDetailCustomizationInstance::CreateStatic(&FPropertyAnimatorCoreEditorDetailCustomization::MakeInstance));
 
 	ISequencerModule& SequencerModule = FModuleManager::Get().LoadModuleChecked<ISequencerModule>("Sequencer");
@@ -27,6 +30,7 @@ void FPropertyAnimatorCoreEditorModule::ShutdownModule()
 	{
 		FPropertyEditorModule& PropertyModule = FModuleManager::GetModuleChecked<FPropertyEditorModule>("PropertyEditor");
 		PropertyModule.UnregisterCustomPropertyTypeLayout(UPropertyAnimatorCoreContext::StaticClass()->GetFName());
+		PropertyModule.UnregisterCustomPropertyTypeLayout(FPropertyAnimatorCoreSequencerTimeSourceChannel::StaticStruct()->GetFName());
 		PropertyModule.UnregisterCustomClassLayout(UPropertyAnimatorCoreBase::StaticClass()->GetFName());
 	}
 
diff --git a/Engine/Plugins/Experimental/PropertyAnimatorCore/Source/PropertyAnimatorCoreEditor/Private/Sequencer/MovieSceneAnimatorTrackEditor.cpp b/Engine/Plugins/Experimental/PropertyAnimatorCore/Source/PropertyAnimatorCoreEditor/Private/Sequencer/MovieSceneAnimatorTrackEditor.cpp
--- a/Engine/Plugins/Experimental/PropertyAnimatorCore/Source/PropertyAnimatorCoreEditor/Private/Sequencer/MovieSceneAnimatorTrackEditor.cpp
+++ b/Engine/Plugins/Experimental/PropertyAnimatorCore/Source/PropertyAnimatorCoreEditor/Private/Sequencer/MovieSceneAnimatorTrackEditor.cpp
@@ -2,23 +2,58 @@
 
 #include "MovieSceneAnimatorTrackEditor.h"
 
+#include "Sequencer/MovieSceneAnimatorSection.h"
 #include "Styling/SlateIconFinder.h"
 
 #define LOCTEXT_NAMESPACE "MovieSceneAnimatorTrackEditor"
 
+FMovieSceneAnimatorTrackEditor::FOnAddAnimatorTrack FMovieSceneAnimatorTrackEditor::OnAddAnimatorTrack;
+FMovieSceneAnimatorTrackEditor::FOnGetAnimatorTrackCount FMovieSceneAnimatorTrackEditor::OnGetAnimatorTrackCount;
+
+FMovieSceneAnimatorTrackEditor::~FMovieSceneAnimatorTrackEditor()
+{
+	OnAddAnimatorTrack.RemoveAll(this);
+	OnGetAnimatorTrackCount.RemoveAll(this);
+}
+
 void FMovieSceneAnimatorTrackEditor::BuildAddTrackMenu(FMenuBuilder& InMenuBuilder)
 {
+	constexpr uint8 Channel = 0;
+
 	InMenuBuilder.AddMenuEntry(
 		LOCTEXT("AddAnimatorTrack.Label", "Animator"),
 		LOCTEXT("AddAnimatorTrack.Tooltip", "Adds a new track that uses the time of the current sequence to drive animators."),
 		FSlateIconFinder::FindIconForClass(UMovieSceneAnimatorTrack::StaticClass()),
 		FUIAction(
-			FExecuteAction::CreateSP(this, &FMovieSceneAnimatorTrackEditor::ExecuteAddTrack)
+			FExecuteAction::CreateSP(this, &FMovieSceneAnimatorTrackEditor::ExecuteAddTrack, Channel)
 		)
 	);
 }
 
-void FMovieSceneAnimatorTrackEditor::ExecuteAddTrack()
+void FMovieSceneAnimatorTrackEditor::BindDelegates()
+{
+	OnAddAnimatorTrack.AddSP(this, &FMovieSceneAnimatorTrackEditor::ExecuteAddTrack);
+	OnGetAnimatorTrackCount.AddSP(this, &FMovieSceneAnimatorTrackEditor::GetTrackCount);
+}
+
+void FMovieSceneAnimatorTrackEditor::GetTrackCount(uint8 InChannel, int32& OutCount) const
+{
+	const UMovieScene* FocusedMovieScene = GetFocusedMovieScene();
+	if (!FocusedMovieScene)
+	{
+		return;
+	}
+
+	for (UMovieSceneTrack* Track : FocusedMovieScene->GetTracks())
+	{
+		if (const UMovieSceneAnimatorTrack* AnimatorTrack = Cast<UMovieSceneAnimatorTrack>(Track))
+		{
+			OutCount += AnimatorTrack->GetChannelCount(InChannel);
+		}
+	}
+}
+
+void FMovieSceneAnimatorTrackEditor::ExecuteAddTrack(uint8 InChannel)
 {
 	UMovieScene* FocusedMovieScene = GetFocusedMovieScene();
 	if (!FocusedMovieScene || FocusedMovieScene->IsReadOnly())
@@ -37,7 +72,9 @@ void FMovieSceneAnimatorTrackEditor::ExecuteAddTrack()
 	FocusedMovieScene->Modify();
 
 	UMovieSceneAnimatorTrack* NewTrack = NewObject<UMovieSceneAnimatorTrack>(FocusedMovieScene, NAME_None, RF_Transactional);
-	NewTrack->AddSection(*NewTrack->CreateNewSection());
+	UMovieSceneAnimatorSection* NewSection = Cast<UMovieSceneAnimatorSection>(NewTrack->CreateNewSection());
+	NewSection->SetChannel(InChannel);
+	NewTrack->AddSection(*NewSection);
 
 	FocusedMovieScene->AddGivenTrack(NewTrack);
 	SequencerPtr->OnAddTrack(NewTrack, FGuid());
diff --git a/Engine/Plugins/Experimental/PropertyAnimatorCore/Source/PropertyAnimatorCoreEditor/Private/Sequencer/MovieSceneAnimatorTrackEditor.h b/Engine/Plugins/Experimental/PropertyAnimatorCore/Source/PropertyAnimatorCoreEditor/Private/Sequencer/MovieSceneAnimatorTrackEditor.h
--- a/Engine/Plugins/Experimental/PropertyAnimatorCore/Source/PropertyAnimatorCoreEditor/Private/Sequencer/MovieSceneAnimatorTrackEditor.h
+++ b/Engine/Plugins/Experimental/PropertyAnimatorCore/Source/PropertyAnimatorCoreEditor/Private/Sequencer/MovieSceneAnimatorTrackEditor.h
@@ -5,16 +5,27 @@
 #include "KeyframeTrackEditor.h"
 #include "Sequencer/MovieSceneAnimatorTrack.h"
 
+/** Animator track editor to add animator track and section */
 class FMovieSceneAnimatorTrackEditor : public FKeyframeTrackEditor<UMovieSceneAnimatorTrack>
 {
 public:
+	DECLARE_MULTICAST_DELEGATE_OneParam(FOnAddAnimatorTrack, uint8 /** Channel */)
+	static FOnAddAnimatorTrack OnAddAnimatorTrack;
+
+	DECLARE_MULTICAST_DELEGATE_TwoParams(FOnGetAnimatorTrackCount, uint8 /** Channel */, int32& /** Count */)
+	static FOnGetAnimatorTrackCount OnGetAnimatorTrackCount;
+
 	FMovieSceneAnimatorTrackEditor(const TSharedRef<ISequencer>& InSequencer)
 		: FKeyframeTrackEditor<UMovieSceneAnimatorTrack>(InSequencer)
 	{}
 
+	virtual ~FMovieSceneAnimatorTrackEditor() override;
+
 	static TSharedRef<ISequencerTrackEditor> CreateTrackEditor(TSharedRef<ISequencer> InSequencer)
 	{
-		return MakeShared<FMovieSceneAnimatorTrackEditor>(InSequencer);
+		TSharedRef<FMovieSceneAnimatorTrackEditor> TrackEditor = MakeShared<FMovieSceneAnimatorTrackEditor>(InSequencer);
+		TrackEditor->BindDelegates();
+		return TrackEditor;
 	}
 
 private:
@@ -22,5 +33,9 @@ private:
 	virtual void BuildAddTrackMenu(FMenuBuilder& InMenuBuilder) override;
 	//~ End FMovieSceneTrackEditor
 
-	void ExecuteAddTrack();
+	void BindDelegates();
+
+	void GetTrackCount(uint8 InChannel, int32& OutCount) const;
+
+	void ExecuteAddTrack(uint8 InChannel);
 };
