MotionDesign : PropertyAnimator
- Display channel in the sequencer animator track
- Fixed sequencer time source channel widget overlapping with entry box

diff --git a/Engine/Plugins/Experimental/PropertyAnimatorCore/Source/PropertyAnimatorCore/Private/Sequencer/MovieSceneAnimatorTrack.cpp b/Engine/Plugins/Experimental/PropertyAnimatorCore/Source/PropertyAnimatorCore/Private/Sequencer/MovieSceneAnimatorTrack.cpp
--- a/Engine/Plugins/Experimental/PropertyAnimatorCore/Source/PropertyAnimatorCore/Private/Sequencer/MovieSceneAnimatorTrack.cpp
+++ b/Engine/Plugins/Experimental/PropertyAnimatorCore/Source/PropertyAnimatorCore/Private/Sequencer/MovieSceneAnimatorTrack.cpp
@@ -90,12 +90,20 @@ void UMovieSceneAnimatorTrack::RemoveSectionAt(int32 InSectionIndex)
 #if WITH_EDITORONLY_DATA
 FText UMovieSceneAnimatorTrack::GetDefaultDisplayName() const
 {
+	if (!Sections.IsEmpty() && Sections[0])
+	{
+		if (const UMovieSceneAnimatorSection* AnimatorSection = Cast<UMovieSceneAnimatorSection>(Sections[0]))
+		{
+			return FText::Format(LOCTEXT("MovieSceneAnimatorTrackChannelName", "Animator Channel {0}"), FText::AsNumber(AnimatorSection->GetChannel()));
+		}
+	}
+
 	return LOCTEXT("MovieSceneAnimatorTrackName", "Animator Channel");
 }
 
 bool UMovieSceneAnimatorTrack::CanRename() const
 {
-	return true;
+	return false;
 }
 #endif
 
diff --git a/Engine/Plugins/Experimental/PropertyAnimatorCore/Source/PropertyAnimatorCoreEditor/Private/Customizations/PropertyAnimatorCoreEditorSequencerTimeSourceChannelTypeCustomization.cpp b/Engine/Plugins/Experimental/PropertyAnimatorCore/Source/PropertyAnimatorCoreEditor/Private/Customizations/PropertyAnimatorCoreEditorSequencerTimeSourceChannelTypeCustomization.cpp
--- a/Engine/Plugins/Experimental/PropertyAnimatorCore/Source/PropertyAnimatorCoreEditor/Private/Customizations/PropertyAnimatorCoreEditorSequencerTimeSourceChannelTypeCustomization.cpp
+++ b/Engine/Plugins/Experimental/PropertyAnimatorCore/Source/PropertyAnimatorCoreEditor/Private/Customizations/PropertyAnimatorCoreEditorSequencerTimeSourceChannelTypeCustomization.cpp
@@ -37,13 +37,15 @@ void FPropertyAnimatorCoreEditorSequencerTimeSourceChannelTypeCustomization::Cus
 	[
 		SNew(SHorizontalBox)
 		+ SHorizontalBox::Slot()
-		.FillWidth(1.f)
+		.AutoWidth()
+		.HAlign(HAlign_Fill)
 		[
 			ChannelPropertyHandle->CreatePropertyValueWidget()
 		]
 		+ SHorizontalBox::Slot()
 		.Padding(5.f, 0.f)
 		.AutoWidth()
+		.HAlign(HAlign_Fill)
 		[
 			SNew(SButton)
 			.HAlign(HAlign_Fill)
@@ -53,7 +55,8 @@ void FPropertyAnimatorCoreEditorSequencerTimeSourceChannelTypeCustomization::Cus
 			[
 				SNew(STextBlock)
 				.Font(IPropertyTypeCustomizationUtils::GetRegularFont())
-				.Text(LOCTEXT("AddSequencerTrack", "Add Sequencer track"))
+				.Text(LOCTEXT("AddSequencerTrack", "Create track"))
+				.ToolTipText(LOCTEXT("AddSequencerTrackTooltip", "Create a sequencer track for this channel"))
 			]
 		]
 	];
