MotionDesign : PropertyAnimator
- Override time source on each animator or use the global time source on the component
- Adding developer settings to set a default time source

diff --git a/Engine/Plugins/Experimental/PropertyAnimatorCore/Source/PropertyAnimatorCore/Private/Animators/PropertyAnimatorCoreBase.cpp b/Engine/Plugins/Experimental/PropertyAnimatorCore/Source/PropertyAnimatorCore/Private/Animators/PropertyAnimatorCoreBase.cpp
--- a/Engine/Plugins/Experimental/PropertyAnimatorCore/Source/PropertyAnimatorCore/Private/Animators/PropertyAnimatorCoreBase.cpp
+++ b/Engine/Plugins/Experimental/PropertyAnimatorCore/Source/PropertyAnimatorCore/Private/Animators/PropertyAnimatorCoreBase.cpp
@@ -7,6 +7,7 @@
 #include "Properties/PropertyAnimatorCoreGroupBase.h"
 #include "Properties/PropertyAnimatorCoreResolver.h"
 #include "Properties/Handlers/PropertyAnimatorCoreHandlerBase.h"
+#include "Settings/PropertyAnimatorCoreSettings.h"
 #include "Subsystems/PropertyAnimatorCoreSubsystem.h"
 #include "TimeSources/PropertyAnimatorCoreTimeSourceBase.h"
 
@@ -19,6 +20,11 @@ UPropertyAnimatorCoreBase::FOnAnimatorPropertyUpdated UPropertyAnimatorCoreBase:
 UPropertyAnimatorCoreBase::FOnAnimatorPropertyUpdated UPropertyAnimatorCoreBase::OnAnimatorPropertyUnlinkedDelegate;
 
 #if WITH_EDITOR
+FName UPropertyAnimatorCoreBase::GetAnimatorEnabledPropertyName()
+{
+	return GET_MEMBER_NAME_CHECKED(UPropertyAnimatorCoreBase, bAnimatorEnabled);
+}
+
 FName UPropertyAnimatorCoreBase::GetLinkedPropertiesPropertyName()
 {
 	return GET_MEMBER_NAME_CHECKED(UPropertyAnimatorCoreBase, LinkedProperties);
@@ -27,13 +33,6 @@ FName UPropertyAnimatorCoreBase::GetLinkedPropertiesPropertyName()
 
 UPropertyAnimatorCoreBase::UPropertyAnimatorCoreBase()
 {
-	if (const UPropertyAnimatorCoreSubsystem* AnimatorSubsystem = UPropertyAnimatorCoreSubsystem::Get())
-	{
-		// Apply default time source
-		const TArray<FName> TimeSourceNames = AnimatorSubsystem->GetTimeSourceNames();
-		SetTimeSourceName(!TimeSourceNames.IsEmpty() ? TimeSourceNames[0] : NAME_None);
-	}
-
 #if WITH_EDITOR
 	if (!IsTemplate())
 	{
@@ -139,6 +138,7 @@ void UPropertyAnimatorCoreBase::PostEditImport()
 {
 	Super::PostEditImport();
 
+	OnTimeSourceNameChanged();
 	ResolvePropertiesOwner();
 }
 
@@ -154,6 +154,7 @@ void UPropertyAnimatorCoreBase::PostDuplicate(EDuplicateMode::Type InMode)
 {
 	Super::PostDuplicate(InMode);
 
+	OnTimeSourceNameChanged();
 	ResolvePropertiesOwner();
 }
 
@@ -184,7 +185,8 @@ void UPropertyAnimatorCoreBase::PostEditChangeProperty(FPropertyChangedEvent& Pr
 	{
 		OnAnimatorEnabledChanged();
 	}
-	else if (MemberName == GET_MEMBER_NAME_CHECKED(UPropertyAnimatorCoreBase, TimeSourceName))
+	else if (MemberName == GET_MEMBER_NAME_CHECKED(UPropertyAnimatorCoreBase, TimeSourceName)
+		|| MemberName == GET_MEMBER_NAME_CHECKED(UPropertyAnimatorCoreBase, bOverrideTimeSource))
 	{
 		OnTimeSourceNameChanged();
 	}
@@ -211,6 +213,17 @@ void UPropertyAnimatorCoreBase::SetAnimatorEnabled(bool bInIsEnabled)
 	OnAnimatorEnabledChanged();
 }
 
+void UPropertyAnimatorCoreBase::SetOverrideTimeSource(bool bInOverride)
+{
+	if (bOverrideTimeSource == bInOverride)
+	{
+		return;
+	}
+
+	bOverrideTimeSource = bInOverride;
+	OnTimeSourceChanged();
+}
+
 void UPropertyAnimatorCoreBase::SetTimeSourceName(FName InTimeSourceName)
 {
 	if (TimeSourceName == InTimeSourceName)
@@ -312,6 +325,16 @@ bool UPropertyAnimatorCoreBase::HasPropertySupport(const FPropertyAnimatorCoreDa
 	return EnumHasAnyFlags(InSupportExpected, GetPropertySupport(InPropertyData));
 }
 
+void UPropertyAnimatorCoreBase::OnAnimatorAdded()
+{
+	bOverrideTimeSource = false;
+
+	if (const UPropertyAnimatorCoreSettings* AnimatorSettings = UPropertyAnimatorCoreSettings::Get())
+	{
+		SetTimeSourceName(AnimatorSettings->GetDefaultTimeSourceName());
+	}
+}
+
 void UPropertyAnimatorCoreBase::OnAnimatorEnabled()
 {
 	UE_LOG(LogPropertyAnimatorCoreBase
@@ -408,28 +431,11 @@ void UPropertyAnimatorCoreBase::ResolvePropertiesOwner(AActor* InNewOwner)
 
 void UPropertyAnimatorCoreBase::EvaluateAnimator(FInstancedPropertyBag& InParameters)
 {
-	UPropertyAnimatorCoreTimeSourceBase* TimeSource = GetActiveTimeSource();
-
-	if (!GetAnimatorEnabled()
-		|| !TimeSource)
-	{
-		return;
-	}
-
-	const TOptional<double> TimeElapsed = TimeSource->GetConditionalTimeElapsed();
-
-	if (!TimeElapsed.IsSet())
-	{
-		return;
-	}
-
 	RestoreProperties();
 
 	SaveProperties();
 
 	EvaluatedPropertyValues.Reset();
-	InParameters.AddProperty(TimeElapsedParameterName, EPropertyBagPropertyType::Double);
-    InParameters.SetValueDouble(TimeElapsedParameterName, TimeElapsed.GetValue());
 
 	bEvaluatingProperties = true;
 	EvaluateProperties(InParameters);
@@ -539,7 +545,7 @@ TArray<FName> UPropertyAnimatorCoreBase::GetTimeSourceNames() const
 
 UPropertyAnimatorCoreTimeSourceBase* UPropertyAnimatorCoreBase::FindOrAddTimeSource(FName InTimeSourceName)
 {
-	if (IsTemplate())
+	if (IsTemplate() || !bOverrideTimeSource)
 	{
 		return nullptr;
 	}
@@ -576,6 +582,22 @@ UPropertyAnimatorCoreTimeSourceBase* UPropertyAnimatorCoreBase::FindOrAddTimeSou
 	return NewTimeSource;
 }
 
+UPropertyAnimatorCoreTimeSourceBase* UPropertyAnimatorCoreBase::GetActiveTimeSource() const
+{
+	UPropertyAnimatorCoreTimeSourceBase* TimeSource = nullptr;
+
+	if (bOverrideTimeSource)
+	{
+		TimeSource = ActiveTimeSource;
+	}
+	else if (const UPropertyAnimatorCoreComponent* AnimatorComponent = GetAnimatorComponent())
+	{
+		TimeSource = AnimatorComponent->GetAnimatorsActiveTimeSource();
+	}
+
+	return TimeSource;
+}
+
 void UPropertyAnimatorCoreBase::SetAnimatorDisplayName(FName InName)
 {
 	if (AnimatorDisplayName == InName)
diff --git a/Engine/Plugins/Experimental/PropertyAnimatorCore/Source/PropertyAnimatorCore/Private/Components/PropertyAnimatorCoreComponent.cpp b/Engine/Plugins/Experimental/PropertyAnimatorCore/Source/PropertyAnimatorCore/Private/Components/PropertyAnimatorCoreComponent.cpp
--- a/Engine/Plugins/Experimental/PropertyAnimatorCore/Source/PropertyAnimatorCore/Private/Components/PropertyAnimatorCoreComponent.cpp
+++ b/Engine/Plugins/Experimental/PropertyAnimatorCore/Source/PropertyAnimatorCore/Private/Components/PropertyAnimatorCoreComponent.cpp
@@ -4,7 +4,9 @@
 
 #include "Engine/World.h"
 #include "GameFramework/Actor.h"
+#include "Settings/PropertyAnimatorCoreSettings.h"
 #include "Subsystems/PropertyAnimatorCoreSubsystem.h"
+#include "TimeSources/PropertyAnimatorCoreTimeSourceBase.h"
 
 UPropertyAnimatorCoreBase* UPropertyAnimatorCoreComponent::AddAnimator(const UClass* InAnimatorClass)
 {
@@ -140,6 +142,60 @@ void UPropertyAnimatorCoreComponent::OnAnimatorsEnabledChanged()
 	SetComponentTickEnabled(bEnableAnimators);
 }
 
+void UPropertyAnimatorCoreComponent::OnTimeSourceNameChanged()
+{
+	if (ActiveAnimatorsTimeSource)
+	{
+		ActiveAnimatorsTimeSource->DeactivateTimeSource();
+	}
+
+	ActiveAnimatorsTimeSource = FindOrAddTimeSource(AnimatorsTimeSourceName);
+
+	if (ActiveAnimatorsTimeSource)
+	{
+		ActiveAnimatorsTimeSource->ActivateTimeSource();
+	}
+}
+
+UPropertyAnimatorCoreTimeSourceBase* UPropertyAnimatorCoreComponent::FindOrAddTimeSource(FName InTimeSourceName)
+{
+	if (IsTemplate())
+	{
+		return nullptr;
+	}
+
+	UPropertyAnimatorCoreSubsystem* Subsystem = UPropertyAnimatorCoreSubsystem::Get();
+
+	if (!Subsystem || InTimeSourceName.IsNone())
+	{
+		return nullptr;
+	}
+
+	// Check cached time source instances
+	UPropertyAnimatorCoreTimeSourceBase* NewTimeSource = nullptr;
+
+	for (const TObjectPtr<UPropertyAnimatorCoreTimeSourceBase>& TimeSource : TimeSources)
+	{
+		if (TimeSource && TimeSource->GetTimeSourceName() == InTimeSourceName)
+		{
+			NewTimeSource = TimeSource.Get();
+		}
+	}
+
+	// Create new time source instance and cache it
+	if (!NewTimeSource)
+	{
+		NewTimeSource = Subsystem->CreateNewTimeSource(InTimeSourceName, this);
+
+		if (NewTimeSource)
+		{
+			TimeSources.Add(NewTimeSource);
+		}
+	}
+
+	return NewTimeSource;
+}
+
 bool UPropertyAnimatorCoreComponent::ShouldAnimate() const
 {
 	return bAnimatorsEnabled
@@ -169,6 +225,11 @@ void UPropertyAnimatorCoreComponent::OnComponentCreated()
 {
 	Super::OnComponentCreated();
 
+	if (const UPropertyAnimatorCoreSettings* AnimatorSettings = UPropertyAnimatorCoreSettings::Get())
+	{
+		SetAnimatorsTimeSourceName(AnimatorSettings->GetDefaultTimeSourceName());
+	}
+
 	if (AActor* OwningActor = GetOwner())
 	{
 		// For spawnable templates, restore and resolve properties owner
@@ -258,14 +319,14 @@ void UPropertyAnimatorCoreComponent::SetAnimatorsMagnitude(float InMagnitude)
 	OnAnimatorsEnabledChanged();
 }
 
-void UPropertyAnimatorCoreComponent::DestroyComponent(bool bPromoteChildren)
+void UPropertyAnimatorCoreComponent::OnComponentDestroyed(bool bInDestroyingHierarchy)
 {
+	Super::OnComponentDestroyed(bInDestroyingHierarchy);
+
 	PropertyAnimatorsInternal = PropertyAnimators;
 	PropertyAnimators.Empty();
 
 	OnAnimatorsChanged();
-
-	Super::DestroyComponent(bPromoteChildren);
 }
 
 void UPropertyAnimatorCoreComponent::TickComponent(float InDeltaTime, ELevelTick InTickType, FActorComponentTickFunction* InTickFunction)
@@ -292,9 +353,35 @@ void UPropertyAnimatorCoreComponent::PostLoad()
 	}
 
 	PRAGMA_ENABLE_DEPRECATION_WARNINGS
+
+	OnTimeSourceNameChanged();
+}
+
+void UPropertyAnimatorCoreComponent::PostEditImport()
+{
+	Super::PostEditImport();
+
+	OnTimeSourceNameChanged();
+}
+
+void UPropertyAnimatorCoreComponent::PostDuplicate(EDuplicateMode::Type InMode)
+{
+	Super::PostDuplicate(InMode);
+
+	OnTimeSourceNameChanged();
 }
 
 #if WITH_EDITOR
+FName UPropertyAnimatorCoreComponent::GetAnimatorsEnabledPropertyName()
+{
+	return GET_MEMBER_NAME_CHECKED(UPropertyAnimatorCoreComponent, bAnimatorsEnabled);
+}
+
+FName UPropertyAnimatorCoreComponent::GetPropertyAnimatorsPropertyName()
+{
+	return GET_MEMBER_NAME_CHECKED(UPropertyAnimatorCoreComponent, PropertyAnimators);
+}
+
 void UPropertyAnimatorCoreComponent::PostEditUndo()
 {
 	Super::PostEditUndo();
@@ -329,6 +416,10 @@ void UPropertyAnimatorCoreComponent::PostEditChangeProperty(FPropertyChangedEven
 	{
 		OnAnimatorsChanged();
 	}
+	else if (MemberName == GET_MEMBER_NAME_CHECKED(UPropertyAnimatorCoreComponent, AnimatorsTimeSourceName))
+	{
+		OnTimeSourceNameChanged();
+	}
 }
 #endif
 
@@ -339,6 +430,23 @@ void UPropertyAnimatorCoreComponent::SetAnimators(const TArray<TObjectPtr<UPrope
 	OnAnimatorsChanged();
 }
 
+void UPropertyAnimatorCoreComponent::SetAnimatorsTimeSourceName(FName InTimeSourceName)
+{
+	if (AnimatorsTimeSourceName == InTimeSourceName)
+	{
+		return;
+	}
+
+	const TArray<FName> TimeSourceNames = GetTimeSourceNames();
+	if (!TimeSourceNames.Contains(InTimeSourceName))
+	{
+		return;
+	}
+
+	AnimatorsTimeSourceName = InTimeSourceName;
+	OnTimeSourceNameChanged();
+}
+
 void UPropertyAnimatorCoreComponent::ForEachAnimator(const TFunctionRef<bool(UPropertyAnimatorCoreBase*)> InFunction) const
 {
 	for (const TObjectPtr<UPropertyAnimatorCoreBase>& Animator : PropertyAnimators)
@@ -377,13 +485,43 @@ bool UPropertyAnimatorCoreComponent::EvaluateAnimators()
 			continue;
 		}
 
+		UPropertyAnimatorCoreTimeSourceBase* AnimatorTimeSource = Animator->GetActiveTimeSource();
+
+		if (!AnimatorTimeSource)
+		{
+			continue;
+		}
+
+		const TOptional<double> TimeElapsed = AnimatorTimeSource->GetConditionalTimeElapsed();
+
+		if (!TimeElapsed.IsSet())
+		{
+			continue;
+		}
+
 		// Reset in case animator change values to avoid affecting following animators
 		Parameters.Reset();
+
 		Parameters.AddProperty(UPropertyAnimatorCoreBase::MagnitudeParameterName, EPropertyBagPropertyType::Float);
 		Parameters.SetValueFloat(UPropertyAnimatorCoreBase::MagnitudeParameterName, AnimatorsMagnitude);
 
+		Parameters.AddProperty(UPropertyAnimatorCoreBase::TimeElapsedParameterName, EPropertyBagPropertyType::Double);
+		Parameters.SetValueDouble(UPropertyAnimatorCoreBase::TimeElapsedParameterName, TimeElapsed.GetValue());
+
 		Animator->EvaluateAnimator(Parameters);
 	}
 
 	return true;
 }
+
+TArray<FName> UPropertyAnimatorCoreComponent::GetTimeSourceNames() const
+{
+	TArray<FName> TimeSourceNames;
+
+	if (const UPropertyAnimatorCoreSubsystem* AnimatorSubsystem = UPropertyAnimatorCoreSubsystem::Get())
+	{
+		TimeSourceNames = AnimatorSubsystem->GetTimeSourceNames();
+	}
+
+	return TimeSourceNames;
+}
diff --git a/Engine/Plugins/Experimental/PropertyAnimatorCore/Source/PropertyAnimatorCore/Private/Settings/PropertyAnimatorCoreSettings.cpp b/Engine/Plugins/Experimental/PropertyAnimatorCore/Source/PropertyAnimatorCore/Private/Settings/PropertyAnimatorCoreSettings.cpp
new file mode 100644
--- /dev/null
+++ b/Engine/Plugins/Experimental/PropertyAnimatorCore/Source/PropertyAnimatorCore/Private/Settings/PropertyAnimatorCoreSettings.cpp
@@ -0,0 +1,48 @@
+// Copyright Epic Games, Inc. All Rights Reserved.
+
+#include "Settings/PropertyAnimatorCoreSettings.h"
+
+#include "Subsystems/PropertyAnimatorCoreSubsystem.h"
+
+UPropertyAnimatorCoreSettings* UPropertyAnimatorCoreSettings::Get()
+{
+	return GetMutableDefault<UPropertyAnimatorCoreSettings>();
+}
+
+UPropertyAnimatorCoreSettings::UPropertyAnimatorCoreSettings()
+{
+	CategoryName = TEXT("Motion Design");
+	SectionName = TEXT("Property Animator Core");
+
+	if (const UPropertyAnimatorCoreSubsystem* AnimatorSubsystem = UPropertyAnimatorCoreSubsystem::Get())
+	{
+		// Apply default time source
+		const TArray<FName> TimeSourceNames = AnimatorSubsystem->GetTimeSourceNames();
+		DefaultTimeSourceName = !TimeSourceNames.IsEmpty() ? TimeSourceNames[0] : NAME_None;
+	}
+}
+
+FName UPropertyAnimatorCoreSettings::GetDefaultTimeSourceName() const
+{
+	const TArray<FName> TimeSources = GetTimeSourceNames();
+
+	if (TimeSources.Contains(DefaultTimeSourceName))
+	{
+		return DefaultTimeSourceName;
+	}
+
+	return !TimeSources.IsEmpty() ? TimeSources[0] : NAME_None;
+}
+
+TArray<FName> UPropertyAnimatorCoreSettings::GetTimeSourceNames() const
+{
+	TArray<FName> TimeSourceNames;
+
+	if (const UPropertyAnimatorCoreSubsystem* AnimatorSubsystem = UPropertyAnimatorCoreSubsystem::Get())
+	{
+		TimeSourceNames = AnimatorSubsystem->GetTimeSourceNames();
+	}
+
+	return TimeSourceNames;
+}
+
diff --git a/Engine/Plugins/Experimental/PropertyAnimatorCore/Source/PropertyAnimatorCore/Private/Settings/PropertyAnimatorCoreSettings.h b/Engine/Plugins/Experimental/PropertyAnimatorCore/Source/PropertyAnimatorCore/Private/Settings/PropertyAnimatorCoreSettings.h
new file mode 100644
--- /dev/null
+++ b/Engine/Plugins/Experimental/PropertyAnimatorCore/Source/PropertyAnimatorCore/Private/Settings/PropertyAnimatorCoreSettings.h
@@ -0,0 +1,28 @@
+// Copyright Epic Games, Inc. All Rights Reserved.
+
+#pragma once
+
+#include "Engine/DeveloperSettings.h"
+#include "PropertyAnimatorCoreSettings.generated.h"
+
+/** Settings for motion design PropertyAnimatorCore plugin */
+UCLASS(Config=Engine, meta=(DisplayName="Property Animator Core"))
+class UPropertyAnimatorCoreSettings : public UDeveloperSettings
+{
+	GENERATED_BODY()
+
+public:
+	static UPropertyAnimatorCoreSettings* Get();
+
+	UPropertyAnimatorCoreSettings();
+
+	FName GetDefaultTimeSourceName() const;
+
+protected:
+	UFUNCTION()
+	TArray<FName> GetTimeSourceNames() const;
+
+	/** The default time source applied on animator by default*/
+	UPROPERTY(Config, EditAnywhere, Category="Animator", meta=(GetOptions="GetTimeSourceNames"))
+	FName DefaultTimeSourceName = NAME_None;
+};
diff --git a/Engine/Plugins/Experimental/PropertyAnimatorCore/Source/PropertyAnimatorCore/Private/Subsystems/PropertyAnimatorCoreSubsystem.cpp b/Engine/Plugins/Experimental/PropertyAnimatorCore/Source/PropertyAnimatorCore/Private/Subsystems/PropertyAnimatorCoreSubsystem.cpp
--- a/Engine/Plugins/Experimental/PropertyAnimatorCore/Source/PropertyAnimatorCore/Private/Subsystems/PropertyAnimatorCoreSubsystem.cpp
+++ b/Engine/Plugins/Experimental/PropertyAnimatorCore/Source/PropertyAnimatorCore/Private/Subsystems/PropertyAnimatorCoreSubsystem.cpp
@@ -888,9 +888,9 @@ UPropertyAnimatorCoreTimeSourceBase* UPropertyAnimatorCoreSubsystem::GetTimeSour
 	return nullptr;
 }
 
-UPropertyAnimatorCoreTimeSourceBase* UPropertyAnimatorCoreSubsystem::CreateNewTimeSource(FName InTimeSourceName, UPropertyAnimatorCoreBase* InAnimator)
+UPropertyAnimatorCoreTimeSourceBase* UPropertyAnimatorCoreSubsystem::CreateNewTimeSource(FName InTimeSourceName, UObject* InOwner)
 {
-	if (!IsValid(InAnimator) || InTimeSourceName.IsNone())
+	if (!IsValid(InOwner) || InTimeSourceName.IsNone())
 	{
 		return nullptr;
 	}
@@ -904,8 +904,8 @@ UPropertyAnimatorCoreTimeSourceBase* UPropertyAnimatorCoreSubsystem::CreateNewTi
 
 	// Here unique name needs to be provided
 	const UClass* TimeSourceClass = TimeSource->GetClass();
-	const FName UniqueObjectName = MakeUniqueObjectName(InAnimator, TimeSourceClass, InTimeSourceName);
-	return NewObject<UPropertyAnimatorCoreTimeSourceBase>(InAnimator, TimeSourceClass, UniqueObjectName);
+	const FName UniqueObjectName = MakeUniqueObjectName(InOwner, TimeSourceClass, InTimeSourceName);
+	return NewObject<UPropertyAnimatorCoreTimeSourceBase>(InOwner, TimeSourceClass, UniqueObjectName);
 }
 
 bool UPropertyAnimatorCoreSubsystem::RegisterPresetClass(const UClass* InPresetClass)
diff --git a/Engine/Plugins/Experimental/PropertyAnimatorCore/Source/PropertyAnimatorCore/Private/TimeSources/PropertyAnimatorCoreTimeSourceBase.cpp b/Engine/Plugins/Experimental/PropertyAnimatorCore/Source/PropertyAnimatorCore/Private/TimeSources/PropertyAnimatorCoreTimeSourceBase.cpp
--- a/Engine/Plugins/Experimental/PropertyAnimatorCore/Source/PropertyAnimatorCore/Private/TimeSources/PropertyAnimatorCoreTimeSourceBase.cpp
+++ b/Engine/Plugins/Experimental/PropertyAnimatorCore/Source/PropertyAnimatorCore/Private/TimeSources/PropertyAnimatorCoreTimeSourceBase.cpp
@@ -2,8 +2,6 @@
 
 #include "TimeSources/PropertyAnimatorCoreTimeSourceBase.h"
 
-#include "Animators/PropertyAnimatorCoreBase.h"
-
 void UPropertyAnimatorCoreTimeSourceBase::ActivateTimeSource()
 {
 	if (IsTimeSourceActive())
@@ -45,11 +43,6 @@ TOptional<double> UPropertyAnimatorCoreTimeSourceBase::GetConditionalTimeElapsed
 	return LastTimeElapsed;
 }
 
-UPropertyAnimatorCoreBase* UPropertyAnimatorCoreTimeSourceBase::GetAnimator() const
-{
-	return GetTypedOuter<UPropertyAnimatorCoreBase>();
-}
-
 void UPropertyAnimatorCoreTimeSourceBase::SetFrameRate(float InFrameRate)
 {
 	FrameRate = FMath::Max(UE_KINDA_SMALL_NUMBER, InFrameRate);
diff --git a/Engine/Plugins/Experimental/PropertyAnimatorCore/Source/PropertyAnimatorCore/Private/TimeSources/PropertyAnimatorCoreWorldTimeSource.cpp b/Engine/Plugins/Experimental/PropertyAnimatorCore/Source/PropertyAnimatorCore/Private/TimeSources/PropertyAnimatorCoreWorldTimeSource.cpp
--- a/Engine/Plugins/Experimental/PropertyAnimatorCore/Source/PropertyAnimatorCore/Private/TimeSources/PropertyAnimatorCoreWorldTimeSource.cpp
+++ b/Engine/Plugins/Experimental/PropertyAnimatorCore/Source/PropertyAnimatorCore/Private/TimeSources/PropertyAnimatorCoreWorldTimeSource.cpp
@@ -17,14 +17,7 @@ double UPropertyAnimatorCoreWorldTimeSource::GetTimeElapsed()
 
 bool UPropertyAnimatorCoreWorldTimeSource::IsTimeSourceReady() const
 {
-	const UPropertyAnimatorCoreBase* Animator = GetAnimator();
-
-	if (!Animator)
-	{
-		return false;
-	}
-
-	const UWorld* World = Animator->GetWorld();
+	const UWorld* World = GetWorld();
 
 	if (!World)
 	{
diff --git a/Engine/Plugins/Experimental/PropertyAnimatorCore/Source/PropertyAnimatorCore/PropertyAnimatorCore.Build.cs b/Engine/Plugins/Experimental/PropertyAnimatorCore/Source/PropertyAnimatorCore/PropertyAnimatorCore.Build.cs
--- a/Engine/Plugins/Experimental/PropertyAnimatorCore/Source/PropertyAnimatorCore/PropertyAnimatorCore.Build.cs
+++ b/Engine/Plugins/Experimental/PropertyAnimatorCore/Source/PropertyAnimatorCore/PropertyAnimatorCore.Build.cs
@@ -14,22 +14,26 @@ public class PropertyAnimatorCore : ModuleRules
 				"Core",
 				"CoreUObject",
 				"MovieScene"
-			});
+			}
+		);
 
 		PrivateDependencyModuleNames.AddRange(
 			new string[]
 			{
 				"ApplicationCore",
 				"CoreUObject",
+				"DeveloperSettings",
 				"Engine"
-			});
+			}
+		);
 
 		if (Target.Type == TargetRules.TargetType.Editor)
 		{
 			PrivateDependencyModuleNames.AddRange(new string[]
-			{
-				"UnrealEd"
-			});
+				{
+					"UnrealEd"
+				}
+			);
 		}
 
 		ShortName = "PropAnimCore";
diff --git a/Engine/Plugins/Experimental/PropertyAnimatorCore/Source/PropertyAnimatorCore/Public/Animators/PropertyAnimatorCoreBase.h b/Engine/Plugins/Experimental/PropertyAnimatorCore/Source/PropertyAnimatorCore/Public/Animators/PropertyAnimatorCoreBase.h
--- a/Engine/Plugins/Experimental/PropertyAnimatorCore/Source/PropertyAnimatorCore/Public/Animators/PropertyAnimatorCoreBase.h
+++ b/Engine/Plugins/Experimental/PropertyAnimatorCore/Source/PropertyAnimatorCore/Public/Animators/PropertyAnimatorCoreBase.h
@@ -28,8 +28,6 @@ class UPropertyAnimatorCoreBase : public UObject
 
 	friend class UPropertyAnimatorCoreComponent;
 	friend class UPropertyAnimatorCoreContext;
-	friend class UPropertyAnimatorCoreEditorStackCustomization;
-	friend class FPropertyAnimatorCoreEditorDetailCustomization;
 
 public:
 	DECLARE_MULTICAST_DELEGATE_OneParam(FOnAnimatorUpdated, UPropertyAnimatorCoreBase* /* InAnimator */)
@@ -57,6 +55,7 @@ public:
 	static constexpr const TCHAR* AlphaParameterName = TEXT("Alpha");
 
 #if WITH_EDITOR
+	PROPERTYANIMATORCORE_API static FName GetAnimatorEnabledPropertyName();
 	PROPERTYANIMATORCORE_API static FName GetLinkedPropertiesPropertyName();
 #endif
 
@@ -73,6 +72,12 @@ public:
 		return bAnimatorEnabled;
 	}
 
+	PROPERTYANIMATORCORE_API void SetOverrideTimeSource(bool bInOverride);
+	bool GetOverrideTimeSource() const
+	{
+		return bOverrideTimeSource;
+	}
+
 	/** Set the time source name to use */
 	PROPERTYANIMATORCORE_API void SetTimeSourceName(FName InTimeSourceName);
 	FName GetTimeSourceName() const
@@ -81,10 +86,7 @@ public:
 	}
 
 	/** Get the active time source */
-	UPropertyAnimatorCoreTimeSourceBase* GetActiveTimeSource() const
-	{
-		return ActiveTimeSource;
-	}
+	PROPERTYANIMATORCORE_API UPropertyAnimatorCoreTimeSourceBase* GetActiveTimeSource() const;
 
 	/** Set the display name of this animator */
 	PROPERTYANIMATORCORE_API void SetAnimatorDisplayName(FName InName);
@@ -146,23 +148,6 @@ public:
 		return Cast<InContextClass>(GetLinkedPropertyContext(InProperty));
 	}
 
-protected:
-	//~ Begin UObject
-	PROPERTYANIMATORCORE_API virtual void BeginDestroy() override;
-	PROPERTYANIMATORCORE_API virtual void PostLoad() override;
-	PROPERTYANIMATORCORE_API virtual void PostEditImport() override;
-	PROPERTYANIMATORCORE_API virtual void PreDuplicate(FObjectDuplicationParameters& InDupParams) override;
-	PROPERTYANIMATORCORE_API virtual void PostDuplicate(EDuplicateMode::Type InDuplicateMode) override;
-#if WITH_EDITOR
-	PROPERTYANIMATORCORE_API virtual void PreEditUndo() override;
-	PROPERTYANIMATORCORE_API virtual void PostEditUndo() override;
-	PROPERTYANIMATORCORE_API virtual void PostEditChangeProperty(FPropertyChangedEvent& PropertyChangedEvent) override;
-#endif
-	//~ End UObject
-
-	/** Update display name based on linked properties */
-	void UpdateAnimatorDisplayName();
-
 	/** Use this to process each linked properties and resolve it, even virtual ones */
 	template<typename InContextClass
 		UE_REQUIRES(TIsDerivedFrom<InContextClass, UPropertyAnimatorCoreContext>::Value)>
@@ -200,6 +185,23 @@ protected:
 		return true;
 	}
 
+protected:
+	//~ Begin UObject
+	PROPERTYANIMATORCORE_API virtual void BeginDestroy() override;
+	PROPERTYANIMATORCORE_API virtual void PostLoad() override;
+	PROPERTYANIMATORCORE_API virtual void PostEditImport() override;
+	PROPERTYANIMATORCORE_API virtual void PreDuplicate(FObjectDuplicationParameters& InDupParams) override;
+	PROPERTYANIMATORCORE_API virtual void PostDuplicate(EDuplicateMode::Type InDuplicateMode) override;
+#if WITH_EDITOR
+	PROPERTYANIMATORCORE_API virtual void PreEditUndo() override;
+	PROPERTYANIMATORCORE_API virtual void PostEditUndo() override;
+	PROPERTYANIMATORCORE_API virtual void PostEditChangeProperty(FPropertyChangedEvent& PropertyChangedEvent) override;
+#endif
+	//~ End UObject
+
+	/** Update display name based on linked properties */
+	void UpdateAnimatorDisplayName();
+
 	/** Used to evaluate linked properties, assign the result in the property bag and return true on success to update property value */
 	template<typename InContextClass
 		UE_REQUIRES(TIsDerivedFrom<InContextClass, UPropertyAnimatorCoreContext>::Value)>
@@ -245,7 +247,7 @@ protected:
 
 	virtual void OnAnimatorDisplayNameChanged() {}
 
-	virtual void OnAnimatorAdded() {}
+	PROPERTYANIMATORCORE_API virtual void OnAnimatorAdded();
 	virtual void OnAnimatorRemoved() {}
 
 	PROPERTYANIMATORCORE_API virtual void OnAnimatorEnabled();
@@ -306,12 +308,16 @@ private:
 	UPROPERTY(EditInstanceOnly, NoClear, Export, Instanced, Category="Animator")
 	TArray<TObjectPtr<UPropertyAnimatorCoreGroupBase>> PropertyGroups;
 
+	/** Use the global time source or override it on this animator */
+	UPROPERTY(EditInstanceOnly, Setter="SetOverrideTimeSource", Getter="GetOverrideTimeSource", Category="Animator", meta=(InlineEditConditionToggle))
+	bool bOverrideTimeSource = true;
+
 	/** The time source to use */
-	UPROPERTY(EditInstanceOnly, Setter="SetTimeSourceName", Getter="GetTimeSourceName", Category="Animator", meta=(GetOptions="GetTimeSourceNames"))
+	UPROPERTY(EditInstanceOnly, Setter, Getter, Category="Animator", meta=(GetOptions="GetTimeSourceNames", EditCondition="bOverrideTimeSource"))
 	FName TimeSourceName = NAME_None;
 
 	/** Active time source with its options, determined by its name */
-	UPROPERTY(VisibleInstanceOnly, Instanced, Transient, DuplicateTransient, Category="Animator")
+	UPROPERTY(VisibleInstanceOnly, Instanced, Transient, DuplicateTransient, Category="Animator", meta=(EditCondition="bOverrideTimeSource", EditConditionHides, HideEditConditionToggle))
 	TObjectPtr<UPropertyAnimatorCoreTimeSourceBase> ActiveTimeSource;
 
 	/** The cached time source used by this Animator */
diff --git a/Engine/Plugins/Experimental/PropertyAnimatorCore/Source/PropertyAnimatorCore/Public/Components/PropertyAnimatorCoreComponent.h b/Engine/Plugins/Experimental/PropertyAnimatorCore/Source/PropertyAnimatorCore/Public/Components/PropertyAnimatorCoreComponent.h
--- a/Engine/Plugins/Experimental/PropertyAnimatorCore/Source/PropertyAnimatorCore/Public/Components/PropertyAnimatorCoreComponent.h
+++ b/Engine/Plugins/Experimental/PropertyAnimatorCore/Source/PropertyAnimatorCore/Public/Components/PropertyAnimatorCoreComponent.h
@@ -7,18 +7,22 @@
 #include "PropertyAnimatorCoreComponent.generated.h"
 
 /** A container for controllers that holds properties in this actor */
-UCLASS(MinimalAPI, ClassGroup=(Custom), AutoExpandCategories=("Animator"), HideCategories=("Activation", "Cooking", "AssetUserData", "Collision"), meta=(BlueprintSpawnableComponent))
+UCLASS(MinimalAPI, ClassGroup=(Custom), AutoExpandCategories=("Animator"), HideCategories=("Activation", "Cooking", "AssetUserData", "Collision", "Tags", "ComponentReplication", "Navigation", "Variable", "Replication"), meta=(BlueprintSpawnableComponent))
 class UPropertyAnimatorCoreComponent : public UActorComponent
 {
 	GENERATED_BODY()
 
 	friend class UPropertyAnimatorCoreSubsystem;
-	friend class UPropertyAnimatorCoreEditorStackCustomization;
 
 public:
 	/** Create an instance of this component class and adds it to an actor */
 	static UPropertyAnimatorCoreComponent* FindOrAdd(AActor* InActor);
 
+#if WITH_EDITOR
+	PROPERTYANIMATORCORE_API static FName GetAnimatorsEnabledPropertyName();
+	PROPERTYANIMATORCORE_API static FName GetPropertyAnimatorsPropertyName();
+#endif
+
 	UPropertyAnimatorCoreComponent();
 
 	void SetAnimators(const TArray<TObjectPtr<UPropertyAnimatorCoreBase>>& InAnimators);
@@ -46,6 +50,17 @@ public:
 		return AnimatorsMagnitude;
 	}
 
+	void SetAnimatorsTimeSourceName(FName InTimeSourceName);
+	FName GetAnimatorsTimeSourceName() const
+	{
+		return AnimatorsTimeSourceName;
+	}
+
+	UPropertyAnimatorCoreTimeSourceBase* GetAnimatorsActiveTimeSource() const
+	{
+		return ActiveAnimatorsTimeSource;
+	}
+
 	/** Process a function for each controller, stops when false is returned otherwise continue until the end */
 	PROPERTYANIMATORCORE_API void ForEachAnimator(TFunctionRef<bool(UPropertyAnimatorCoreBase*)> InFunction) const;
 
@@ -57,12 +72,14 @@ protected:
 
 	//~ Begin UActorComponent
 	virtual void OnComponentCreated() override;
-	virtual void DestroyComponent(bool bPromoteChildren) override;
+	virtual void OnComponentDestroyed(bool bInDestroyingHierarchy) override;
 	virtual void TickComponent(float InDeltaTime, ELevelTick InTickType, FActorComponentTickFunction* InTickFunction) override;
 	//~ End UActorComponent
 
 	//~ Begin UObject
 	virtual void PostLoad() override;
+	virtual void PostEditImport() override;
+	virtual void PostDuplicate(EDuplicateMode::Type InMode) override;
 #if WITH_EDITOR
 	virtual void PostEditUndo() override;
 	virtual void PreEditChange(FProperty* PropertyAboutToChange) override;
@@ -96,9 +113,18 @@ protected:
 	/** Callback when global enabled state is changed */
 	void OnAnimatorsEnabledChanged();
 
+	/** Callback when global time source name is changed */
+	void OnTimeSourceNameChanged();
+
 	/** Evaluate only specified animators */
 	bool EvaluateAnimators();
 
+	/** Finds a cached time source with this name or creates a new one */
+	UPropertyAnimatorCoreTimeSourceBase* FindOrAddTimeSource(FName InTimeSourceName);
+
+	UFUNCTION()
+	TArray<FName> GetTimeSourceNames() const;
+
 	/** Animators linked to this actor, they contain only properties within this actor */
 	UPROPERTY(EditAnywhere, BlueprintReadWrite, NoClear, Export, Instanced, Setter="SetAnimators", Category="Animator", meta=(TitleProperty="AnimatorDisplayName"))
 	TArray<TObjectPtr<UPropertyAnimatorCoreBase>> PropertyAnimators;
@@ -111,6 +137,14 @@ protected:
 	UPROPERTY(EditInstanceOnly, Getter, Setter, Category="Animator", meta=(ClampMin="0", ClampMax="1", UIMin="0", UIMax="1", AllowPrivateAccess="true"))
 	float AnimatorsMagnitude = 1.f;
 
+	/** The global time source to use, can be overriden in animator */
+	UPROPERTY(EditInstanceOnly, Setter, Getter, Category="Animator", meta=(GetOptions="GetTimeSourceNames"))
+	FName AnimatorsTimeSourceName = NAME_None;
+
+	/** Active time source with its options, determined by its name */
+	UPROPERTY(VisibleInstanceOnly, Instanced, Transient, DuplicateTransient, Category="Animator")
+	TObjectPtr<UPropertyAnimatorCoreTimeSourceBase> ActiveAnimatorsTimeSource;
+
 private:
 	/** Deprecated property set, will be migrated to PropertyAnimators property on load */
 	UE_DEPRECATED(5.5, "Moved to PropertyAnimators")
@@ -123,5 +157,5 @@ private:
 
 	/** Cached time sources used by this animator component */
 	UPROPERTY()
-	TArray<TObjectPtr<UPropertyAnimatorCoreTimeSourceBase>> TimeSourceInstances;
+	TArray<TObjectPtr<UPropertyAnimatorCoreTimeSourceBase>> TimeSources;
 };
\ No newline at end of file
diff --git a/Engine/Plugins/Experimental/PropertyAnimatorCore/Source/PropertyAnimatorCore/Public/Subsystems/PropertyAnimatorCoreSubsystem.h b/Engine/Plugins/Experimental/PropertyAnimatorCore/Source/PropertyAnimatorCore/Public/Subsystems/PropertyAnimatorCoreSubsystem.h
--- a/Engine/Plugins/Experimental/PropertyAnimatorCore/Source/PropertyAnimatorCore/Public/Subsystems/PropertyAnimatorCoreSubsystem.h
+++ b/Engine/Plugins/Experimental/PropertyAnimatorCore/Source/PropertyAnimatorCore/Public/Subsystems/PropertyAnimatorCoreSubsystem.h
@@ -100,7 +100,7 @@ public:
 	UPropertyAnimatorCoreTimeSourceBase* GetTimeSource(FName InTimeSourceName) const;
 
 	/** Create a new time source for an animator */
-	UPropertyAnimatorCoreTimeSourceBase* CreateNewTimeSource(FName InTimeSourceName, UPropertyAnimatorCoreBase* InAnimator);
+	UPropertyAnimatorCoreTimeSourceBase* CreateNewTimeSource(FName InTimeSourceName, UObject* InOwner);
 
 	/** Register a preset class */
 	PROPERTYANIMATORCORE_API bool RegisterPresetClass(const UClass* InPresetClass);
diff --git a/Engine/Plugins/Experimental/PropertyAnimatorCore/Source/PropertyAnimatorCore/Public/TimeSources/PropertyAnimatorCoreTimeSourceBase.h b/Engine/Plugins/Experimental/PropertyAnimatorCore/Source/PropertyAnimatorCore/Public/TimeSources/PropertyAnimatorCoreTimeSourceBase.h
--- a/Engine/Plugins/Experimental/PropertyAnimatorCore/Source/PropertyAnimatorCore/Public/TimeSources/PropertyAnimatorCoreTimeSourceBase.h
+++ b/Engine/Plugins/Experimental/PropertyAnimatorCore/Source/PropertyAnimatorCore/Public/TimeSources/PropertyAnimatorCoreTimeSourceBase.h
@@ -36,9 +36,6 @@ public:
 
 	TOptional<double> GetConditionalTimeElapsed();
 
-	/** Get the animator, this time source is on */
-	PROPERTYANIMATORCORE_API UPropertyAnimatorCoreBase* GetAnimator() const;
-
 	FName GetTimeSourceName() const
 	{
 		return TimeSourceName;
diff --git a/Engine/Plugins/Experimental/PropertyAnimatorCore/Source/PropertyAnimatorCoreEditor/Private/Customizations/PropertyAnimatorCoreEditorDetailCustomization.cpp b/Engine/Plugins/Experimental/PropertyAnimatorCore/Source/PropertyAnimatorCoreEditor/Private/Customizations/PropertyAnimatorCoreEditorDetailCustomization.cpp
--- a/Engine/Plugins/Experimental/PropertyAnimatorCore/Source/PropertyAnimatorCoreEditor/Private/Customizations/PropertyAnimatorCoreEditorDetailCustomization.cpp
+++ b/Engine/Plugins/Experimental/PropertyAnimatorCore/Source/PropertyAnimatorCoreEditor/Private/Customizations/PropertyAnimatorCoreEditorDetailCustomization.cpp
@@ -19,7 +19,7 @@
 void FPropertyAnimatorCoreEditorDetailCustomization::CustomizeDetails(IDetailLayoutBuilder& InDetailBuilder)
 {
 	const TSharedPtr<IPropertyHandle> LinkedPropertiesHandle = InDetailBuilder.GetProperty(
-		GET_MEMBER_NAME_CHECKED(UPropertyAnimatorCoreBase, LinkedProperties),
+		UPropertyAnimatorCoreBase::GetLinkedPropertiesPropertyName(),
 		UPropertyAnimatorCoreBase::StaticClass()
 	);
 
@@ -29,13 +29,15 @@ void FPropertyAnimatorCoreEditorDetailCustomization::CustomizeDetails(IDetailLay
 	}
 
 	TArray<TWeakObjectPtr<UPropertyAnimatorCoreBase>> AnimatorsWeak = InDetailBuilder.GetObjectsOfTypeBeingCustomized<UPropertyAnimatorCoreBase>();
+	UPropertyAnimatorCoreBase* Animator = AnimatorsWeak.IsValidIndex(0) ? AnimatorsWeak[0].Get() : nullptr;
 
-	if (AnimatorsWeak.Num() != 1 || !AnimatorsWeak[0].IsValid())
+	if (!Animator)
 	{
 		return;
 	}
 
-	AnimatorWeak = AnimatorsWeak[0].Get();
+	AnimatorWeak = Animator;
+
 	IDetailPropertyRow* PropertyRow = InDetailBuilder.EditDefaultProperty(LinkedPropertiesHandle);
 
 	if (!PropertyRow)
diff --git a/Engine/Plugins/Experimental/PropertyAnimatorCore/Source/PropertyAnimatorCoreEditor/Private/Customizations/PropertyAnimatorCoreEditorStackCustomization.cpp b/Engine/Plugins/Experimental/PropertyAnimatorCore/Source/PropertyAnimatorCoreEditor/Private/Customizations/PropertyAnimatorCoreEditorStackCustomization.cpp
--- a/Engine/Plugins/Experimental/PropertyAnimatorCore/Source/PropertyAnimatorCoreEditor/Private/Customizations/PropertyAnimatorCoreEditorStackCustomization.cpp
+++ b/Engine/Plugins/Experimental/PropertyAnimatorCore/Source/PropertyAnimatorCoreEditor/Private/Customizations/PropertyAnimatorCoreEditorStackCustomization.cpp
@@ -12,7 +12,6 @@
 #include "Styling/SlateIconFinder.h"
 #include "Subsystems/PropertyAnimatorCoreEditorSubsystem.h"
 #include "Subsystems/PropertyAnimatorCoreSubsystem.h"
-#include "TimeSources/PropertyAnimatorCoreTimeSourceBase.h"
 #include "ToolMenus.h"
 #include "ToolMenu.h"
 
@@ -177,7 +176,7 @@ void UPropertyAnimatorCoreEditorStackCustomization::CustomizeItemHeader(const FO
 	// Customize component header
 	if (InItem->IsA<UPropertyAnimatorCoreComponent>())
 	{
-		FBoolProperty* EnableProperty = FindFProperty<FBoolProperty>(UPropertyAnimatorCoreComponent::StaticClass(), GET_MEMBER_NAME_CHECKED(UPropertyAnimatorCoreComponent, bAnimatorsEnabled));
+		FBoolProperty* EnableProperty = FindFProperty<FBoolProperty>(UPropertyAnimatorCoreComponent::StaticClass(), UPropertyAnimatorCoreComponent::GetAnimatorsEnabledPropertyName());
 
 		const FSlateIcon ClassIcon = FSlateIconFinder::FindIconForClass(UPropertyAnimatorCoreComponent::StaticClass());
 
@@ -209,7 +208,7 @@ void UPropertyAnimatorCoreEditorStackCustomization::CustomizeItemHeader(const FO
 	{
 		UPropertyAnimatorCoreBase* Animator = InItem->Get<UPropertyAnimatorCoreBase>(0);
 
-		FBoolProperty* EnableProperty = FindFProperty<FBoolProperty>(UPropertyAnimatorCoreBase::StaticClass(), GET_MEMBER_NAME_CHECKED(UPropertyAnimatorCoreBase, bAnimatorEnabled));
+		FBoolProperty* EnableProperty = FindFProperty<FBoolProperty>(UPropertyAnimatorCoreBase::StaticClass(), UPropertyAnimatorCoreBase::GetAnimatorEnabledPropertyName());
 
 		const FSlateIcon ClassIcon = FSlateIconFinder::FindIconForClass(Animator->GetClass());
 
@@ -284,16 +283,18 @@ void UPropertyAnimatorCoreEditorStackCustomization::CustomizeItemBody(const FOpe
 {
 	if (InItem->IsA<UPropertyAnimatorCoreComponent>())
 	{
-		FProperty* MagnitudeProperty = FindFProperty<FProperty>(UPropertyAnimatorCoreComponent::StaticClass(), GET_MEMBER_NAME_CHECKED(UPropertyAnimatorCoreComponent, AnimatorsMagnitude));
+		FProperty* AnimatorsEnabledProperty = FindFProperty<FProperty>(UPropertyAnimatorCoreComponent::StaticClass(), UPropertyAnimatorCoreComponent::GetAnimatorsEnabledPropertyName());
+		FProperty* PropertyAnimatorsProperty = FindFProperty<FProperty>(UPropertyAnimatorCoreComponent::StaticClass(), UPropertyAnimatorCoreComponent::GetPropertyAnimatorsPropertyName());
 
 		InBodyBuilder
-			.AllowProperty(MagnitudeProperty)
+			.DisallowProperty(AnimatorsEnabledProperty)
+			.DisallowProperty(PropertyAnimatorsProperty)
 			.SetShowDetailsView(true);
 	}
 	// Customize animator body
 	else if (InItem->IsA<UPropertyAnimatorCoreBase>())
 	{
-		FBoolProperty* EnableProperty = FindFProperty<FBoolProperty>(UPropertyAnimatorCoreBase::StaticClass(), GET_MEMBER_NAME_CHECKED(UPropertyAnimatorCoreBase, bAnimatorEnabled));
+		FBoolProperty* EnableProperty = FindFProperty<FBoolProperty>(UPropertyAnimatorCoreBase::StaticClass(), UPropertyAnimatorCoreBase::GetAnimatorEnabledPropertyName());
 		FProperty* LinkedPropertiesProperty = FindFProperty<FProperty>(UPropertyAnimatorCoreBase::StaticClass(), UPropertyAnimatorCoreBase::GetLinkedPropertiesPropertyName());
 
 		InBodyBuilder
diff --git a/Engine/Plugins/Experimental/PropertyAnimatorCore/Source/PropertyAnimatorCoreEditor/Private/PropertyAnimatorCoreEditorModule.cpp b/Engine/Plugins/Experimental/PropertyAnimatorCore/Source/PropertyAnimatorCoreEditor/Private/PropertyAnimatorCoreEditorModule.cpp
--- a/Engine/Plugins/Experimental/PropertyAnimatorCore/Source/PropertyAnimatorCoreEditor/Private/PropertyAnimatorCoreEditorModule.cpp
+++ b/Engine/Plugins/Experimental/PropertyAnimatorCore/Source/PropertyAnimatorCoreEditor/Private/PropertyAnimatorCoreEditorModule.cpp
@@ -15,12 +15,12 @@
 
 void FPropertyAnimatorCoreEditorModule::StartupModule()
 {
-	FPropertyEditorModule& PropertyModule = FModuleManager::GetModuleChecked<FPropertyEditorModule>("PropertyEditor");
+	FPropertyEditorModule& PropertyModule = FModuleManager::LoadModuleChecked<FPropertyEditorModule>("PropertyEditor");
 	PropertyModule.RegisterCustomPropertyTypeLayout(UPropertyAnimatorCoreContext::StaticClass()->GetFName(), FOnGetPropertyTypeCustomizationInstance::CreateStatic(&FPropertyAnimatorCoreEditorContextTypeCustomization::MakeInstance));
 	PropertyModule.RegisterCustomPropertyTypeLayout(FPropertyAnimatorCoreSequencerTimeSourceChannel::StaticStruct()->GetFName(), FOnGetPropertyTypeCustomizationInstance::CreateStatic(&FPropertyAnimatorCoreEditorSequencerTimeSourceChannelTypeCustomization::MakeInstance));
 	PropertyModule.RegisterCustomClassLayout(UPropertyAnimatorCoreBase::StaticClass()->GetFName(), FOnGetDetailCustomizationInstance::CreateStatic(&FPropertyAnimatorCoreEditorDetailCustomization::MakeInstance));
 
-	ISequencerModule& SequencerModule = FModuleManager::Get().LoadModuleChecked<ISequencerModule>("Sequencer");
+	ISequencerModule& SequencerModule = FModuleManager::LoadModuleChecked<ISequencerModule>("Sequencer");
 	AnimatorTrackCreateEditorHandle = SequencerModule.RegisterTrackEditor(FOnCreateTrackEditor::CreateStatic(&FMovieSceneAnimatorTrackEditor::CreateTrackEditor));
 }
 
