MotionDesign : PropertyAnimatorCore & ActorModifierCore
- Removing the component instead of all objects within when top trash button is clicked, this will not leave any residual objects

diff --git a/Engine/Plugins/Experimental/ActorModifierCore/Source/ActorModifierCore/Private/Subsystems/ActorModifierCoreSubsystem.cpp b/Engine/Plugins/Experimental/ActorModifierCore/Source/ActorModifierCore/Private/Subsystems/ActorModifierCoreSubsystem.cpp
--- a/Engine/Plugins/Experimental/ActorModifierCore/Source/ActorModifierCore/Private/Subsystems/ActorModifierCoreSubsystem.cpp
+++ b/Engine/Plugins/Experimental/ActorModifierCore/Source/ActorModifierCore/Private/Subsystems/ActorModifierCoreSubsystem.cpp
@@ -1143,6 +1143,53 @@ bool UActorModifierCoreSubsystem::RemoveModifiers(const TSet<UActorModifierCoreB
 	return true;
 }
 
+bool UActorModifierCoreSubsystem::RemoveModifierStacks(const TSet<UActorModifierCoreStack*>& InStacks, bool bInShouldTransact) const
+{
+	if (InStacks.IsEmpty())
+	{
+		return false;
+	}
+
+#if WITH_EDITOR
+	const FText TransactionText = LOCTEXT("RemoveModifierComponent", "Removing {0} modifier component(s)");
+	const FText StackCount = FText::FromString(FString::FromInt(InStacks.Num()));
+
+	FScopedTransaction Transaction(FText::Format(TransactionText, StackCount), bInShouldTransact);
+#endif
+
+	for (UActorModifierCoreStack* ActorStack : InStacks)
+	{
+		if (!IsValid(ActorStack))
+		{
+			continue;
+		}
+
+		UActorModifierCoreComponent* Component = ActorStack->GetTypedOuter<UActorModifierCoreComponent>();
+
+		if (!IsValid(Component))
+		{
+			continue;
+		}
+
+		AActor* Actor = Component->GetOwner();
+
+		if (!IsValid(Actor))
+		{
+			continue;
+		}
+
+#if WITH_EDITOR
+		Actor->Modify();
+		Component->Modify();
+		ActorStack->Modify();
+#endif
+
+		Component->DestroyComponent(/** PromoteChildren */false);
+	}
+
+	return true;
+}
+
 bool UActorModifierCoreSubsystem::RemoveActorsModifiers(const TSet<AActor*>& InActors, bool bInShouldTransact) const
 {
 	if (InActors.IsEmpty())
diff --git a/Engine/Plugins/Experimental/ActorModifierCore/Source/ActorModifierCore/Public/Subsystems/ActorModifierCoreSubsystem.h b/Engine/Plugins/Experimental/ActorModifierCore/Source/ActorModifierCore/Public/Subsystems/ActorModifierCoreSubsystem.h
--- a/Engine/Plugins/Experimental/ActorModifierCore/Source/ActorModifierCore/Public/Subsystems/ActorModifierCoreSubsystem.h
+++ b/Engine/Plugins/Experimental/ActorModifierCore/Source/ActorModifierCore/Public/Subsystems/ActorModifierCoreSubsystem.h
@@ -121,6 +121,9 @@ public:
 	/** Remove modifiers from different actors or stacks, will update the original array and perform a transaction if wanted */
 	ACTORMODIFIERCORE_API bool RemoveModifiers(const TSet<UActorModifierCoreBase*>& InModifiers, FActorModifierCoreStackRemoveOp& InRemoveOp) const;
 
+	/** Removes modifier components and modifier stacks */
+	ACTORMODIFIERCORE_API bool RemoveModifierStacks(const TSet<UActorModifierCoreStack*>& InStacks, bool bInShouldTransact = false) const;
+
 	/** Insert a modifier in a stack before or after another modifier, will perform a transaction if wanted */
 	ACTORMODIFIERCORE_API UActorModifierCoreBase* InsertModifier(UActorModifierCoreStack* InStack, FActorModifierCoreStackInsertOp& InInsertOp) const;
 
diff --git a/Engine/Plugins/Experimental/ActorModifierCore/Source/ActorModifierCoreEditor/Private/Modifiers/Customizations/ActorModifierCoreEditorStackCustomization.cpp b/Engine/Plugins/Experimental/ActorModifierCore/Source/ActorModifierCoreEditor/Private/Modifiers/Customizations/ActorModifierCoreEditorStackCustomization.cpp
--- a/Engine/Plugins/Experimental/ActorModifierCore/Source/ActorModifierCoreEditor/Private/Modifiers/Customizations/ActorModifierCoreEditorStackCustomization.cpp
+++ b/Engine/Plugins/Experimental/ActorModifierCore/Source/ActorModifierCoreEditor/Private/Modifiers/Customizations/ActorModifierCoreEditorStackCustomization.cpp
@@ -654,16 +654,11 @@ void UActorModifierCoreEditorStackCustomization::RemoveModifierAction(FOperatorS
 
 	if (InItem->IsA<UActorModifierCoreStack>())
 	{
-		TSet<AActor*> Actors;
-
-		for (const UActorModifierCoreStack* ModifierStack : InItem->GetAsArray<UActorModifierCoreStack>())
-		{
-			Actors.Add(ModifierStack->GetModifiedActor());
-		}
+		const TSet<UActorModifierCoreStack*> ModifierStacks(InItem->GetAsArray<UActorModifierCoreStack>());
 
-		if (!ModifierSubsystem->RemoveActorsModifiers(Actors, true))
+		if (!ModifierSubsystem->RemoveModifierStacks(ModifierStacks, /** Transact */true))
 		{
-			UE_LOG(LogActorModifierCoreEditorStackCustomization, Warning, TEXT("Could not remove modifiers from actors"))
+			UE_LOG(LogActorModifierCoreEditorStackCustomization, Warning, TEXT("Could not remove modifier stacks from actors"))
 		}
 	}
 	else
diff --git a/Engine/Plugins/Experimental/Avalanche/Source/AvalancheEditor/Private/OperatorStack/Widgets/SAvaOperatorStackTab.cpp b/Engine/Plugins/Experimental/Avalanche/Source/AvalancheEditor/Private/OperatorStack/Widgets/SAvaOperatorStackTab.cpp
--- a/Engine/Plugins/Experimental/Avalanche/Source/AvalancheEditor/Private/OperatorStack/Widgets/SAvaOperatorStackTab.cpp
+++ b/Engine/Plugins/Experimental/Avalanche/Source/AvalancheEditor/Private/OperatorStack/Widgets/SAvaOperatorStackTab.cpp
@@ -3,6 +3,7 @@
 #include "SAvaOperatorStackTab.h"
 
 #include "Animators/PropertyAnimatorCoreBase.h"
+#include "Components/PropertyAnimatorCoreComponent.h"
 #include "Contexts/OperatorStackEditorContext.h"
 #include "DetailView/IAvaDetailsProvider.h"
 #include "EditorModeManager.h"
@@ -34,9 +35,9 @@ void SAvaOperatorStackTab::Construct(const FArguments& InArgs
 	UActorModifierCoreStack::OnModifierReplaced().AddSP(this, &SAvaOperatorStackTab::OnModifierUpdated);
 
 	// Property controllers delegates
-	UPropertyAnimatorCoreBase::OnAnimatorCreatedDelegate.AddSP(this, &SAvaOperatorStackTab::OnAnimatorUpdated);
-	UPropertyAnimatorCoreBase::OnAnimatorRemovedDelegate.AddSP(this, &SAvaOperatorStackTab::OnAnimatorUpdated);
-	UPropertyAnimatorCoreBase::OnAnimatorRenamedDelegate.AddSP(this, &SAvaOperatorStackTab::OnAnimatorUpdated);
+	UPropertyAnimatorCoreBase::OnPropertyAnimatorAdded().AddSP(this, &SAvaOperatorStackTab::OnAnimatorUpdated);
+	UPropertyAnimatorCoreBase::OnPropertyAnimatorRemoved().AddSP(this, &SAvaOperatorStackTab::OnAnimatorUpdated);
+	UPropertyAnimatorCoreBase::OnPropertyAnimatorRenamed().AddSP(this, &SAvaOperatorStackTab::OnAnimatorUpdated);
 
 	const TSharedPtr<IDetailKeyframeHandler> KeyframeHandler = InProvider->GetDetailsKeyframeHandler();
 
@@ -61,9 +62,9 @@ SAvaOperatorStackTab::~SAvaOperatorStackTab()
 	UActorModifierCoreStack::OnModifierRemoved().RemoveAll(this);
 	UActorModifierCoreStack::OnModifierReplaced().RemoveAll(this);
 
-	UPropertyAnimatorCoreBase::OnAnimatorCreatedDelegate.RemoveAll(this);
-	UPropertyAnimatorCoreBase::OnAnimatorRemovedDelegate.RemoveAll(this);
-	UPropertyAnimatorCoreBase::OnAnimatorRenamedDelegate.RemoveAll(this);
+	UPropertyAnimatorCoreBase::OnPropertyAnimatorAdded().RemoveAll(this);
+	UPropertyAnimatorCoreBase::OnPropertyAnimatorRemoved().RemoveAll(this);
+	UPropertyAnimatorCoreBase::OnPropertyAnimatorRenamed().RemoveAll(this);
 }
 
 void SAvaOperatorStackTab::RefreshSelection(UObject* InSelectionObject) const
@@ -102,13 +103,16 @@ void SAvaOperatorStackTab::OnModifierUpdated(UActorModifierCoreBase* InUpdatedIt
 {
 	if (InUpdatedItem)
 	{
-		RefreshCurrentSelection(InUpdatedItem->GetModifierStack());
+		RefreshCurrentSelection(InUpdatedItem->GetRootModifierStack());
 	}
 }
 
-void SAvaOperatorStackTab::OnAnimatorUpdated(UPropertyAnimatorCoreBase* InUpdatedItem) const
+void SAvaOperatorStackTab::OnAnimatorUpdated(UPropertyAnimatorCoreComponent* InComponent, UPropertyAnimatorCoreBase* InUpdatedItem) const
 {
-	RefreshCurrentSelection(InUpdatedItem);
+	if (InComponent)
+	{
+		RefreshCurrentSelection(InComponent);
+	}
 }
 
 void SAvaOperatorStackTab::RefreshCurrentSelection(const UObject* InObject) const
diff --git a/Engine/Plugins/Experimental/Avalanche/Source/AvalancheEditor/Private/OperatorStack/Widgets/SAvaOperatorStackTab.h b/Engine/Plugins/Experimental/Avalanche/Source/AvalancheEditor/Private/OperatorStack/Widgets/SAvaOperatorStackTab.h
--- a/Engine/Plugins/Experimental/Avalanche/Source/AvalancheEditor/Private/OperatorStack/Widgets/SAvaOperatorStackTab.h
+++ b/Engine/Plugins/Experimental/Avalanche/Source/AvalancheEditor/Private/OperatorStack/Widgets/SAvaOperatorStackTab.h
@@ -10,6 +10,7 @@ class IAvaDetailsProvider;
 class IDetailKeyframeHandler;
 class SOperatorStackEditorWidget;
 class UPropertyAnimatorCoreBase;
+class UPropertyAnimatorCoreComponent;
 class UObject;
 class UActorModifierCoreBase;
 
@@ -29,7 +30,7 @@ public:
 private:
 	void RefreshSelection(UObject* InSelectionObject) const;
 	void OnModifierUpdated(UActorModifierCoreBase* InUpdatedItem) const;
-	void OnAnimatorUpdated(UPropertyAnimatorCoreBase* InUpdatedItem) const;
+	void OnAnimatorUpdated(UPropertyAnimatorCoreComponent* InComponent, UPropertyAnimatorCoreBase* InUpdatedItem) const;
 	void RefreshCurrentSelection(const UObject* InObject) const;
 
 	TWeakPtr<IAvaDetailsProvider> DetailsProviderWeak;
diff --git a/Engine/Plugins/Experimental/Avalanche/Source/AvalanchePropertyAnimatorEditor/Private/Outliner/AvaPropertyAnimatorEditorOutlinerProxy.cpp b/Engine/Plugins/Experimental/Avalanche/Source/AvalanchePropertyAnimatorEditor/Private/Outliner/AvaPropertyAnimatorEditorOutlinerProxy.cpp
--- a/Engine/Plugins/Experimental/Avalanche/Source/AvalanchePropertyAnimatorEditor/Private/Outliner/AvaPropertyAnimatorEditorOutlinerProxy.cpp
+++ b/Engine/Plugins/Experimental/Avalanche/Source/AvalanchePropertyAnimatorEditor/Private/Outliner/AvaPropertyAnimatorEditorOutlinerProxy.cpp
@@ -126,27 +126,25 @@ void FAvaPropertyAnimatorEditorOutlinerProxy::GetProxiedItems(const TSharedRef<I
 void FAvaPropertyAnimatorEditorOutlinerProxy::BindDelegates()
 {
 	UnbindDelegates();
-	UPropertyAnimatorCoreBase::OnAnimatorCreatedDelegate.AddSP(this, &FAvaPropertyAnimatorEditorOutlinerProxy::OnPropertyAnimatorUpdated);
-	UPropertyAnimatorCoreBase::OnAnimatorRemovedDelegate.AddSP(this, &FAvaPropertyAnimatorEditorOutlinerProxy::OnPropertyAnimatorUpdated);
-	UPropertyAnimatorCoreBase::OnAnimatorRenamedDelegate.AddSP(this, &FAvaPropertyAnimatorEditorOutlinerProxy::OnPropertyAnimatorUpdated);
+	UPropertyAnimatorCoreBase::OnPropertyAnimatorAdded().AddSP(this, &FAvaPropertyAnimatorEditorOutlinerProxy::OnPropertyAnimatorUpdated);
+	UPropertyAnimatorCoreBase::OnPropertyAnimatorRemoved().AddSP(this, &FAvaPropertyAnimatorEditorOutlinerProxy::OnPropertyAnimatorUpdated);
+	UPropertyAnimatorCoreBase::OnPropertyAnimatorRenamed().AddSP(this, &FAvaPropertyAnimatorEditorOutlinerProxy::OnPropertyAnimatorUpdated);
 }
 
 void FAvaPropertyAnimatorEditorOutlinerProxy::UnbindDelegates()
 {
-	UPropertyAnimatorCoreBase::OnAnimatorCreatedDelegate.RemoveAll(this);
-	UPropertyAnimatorCoreBase::OnAnimatorRemovedDelegate.RemoveAll(this);
-	UPropertyAnimatorCoreBase::OnAnimatorRenamedDelegate.RemoveAll(this);
+	UPropertyAnimatorCoreBase::OnPropertyAnimatorAdded().RemoveAll(this);
+	UPropertyAnimatorCoreBase::OnPropertyAnimatorRemoved().RemoveAll(this);
+	UPropertyAnimatorCoreBase::OnPropertyAnimatorRenamed().RemoveAll(this);
 }
 
-void FAvaPropertyAnimatorEditorOutlinerProxy::OnPropertyAnimatorUpdated(UPropertyAnimatorCoreBase* InAnimator)
+void FAvaPropertyAnimatorEditorOutlinerProxy::OnPropertyAnimatorUpdated(UPropertyAnimatorCoreComponent* InComponent, UPropertyAnimatorCoreBase* InAnimator)
 {
-	const UPropertyAnimatorCoreComponent* PropertyAnimatorComponent = GetPropertyAnimatorComponent();
+	const UPropertyAnimatorCoreComponent* ActiveAnimatorComponent = GetPropertyAnimatorComponent();
 
-	if (IsValid(InAnimator) && IsValid(PropertyAnimatorComponent))
+	if (IsValid(InAnimator) && IsValid(ActiveAnimatorComponent))
 	{
-		const UPropertyAnimatorCoreComponent* UpdatedComponent = InAnimator->GetTypedOuter<UPropertyAnimatorCoreComponent>();
-
-		if (UpdatedComponent == PropertyAnimatorComponent)
+		if (InComponent == ActiveAnimatorComponent)
 		{
 			RefreshChildren();
 			Outliner.RequestRefresh();
diff --git a/Engine/Plugins/Experimental/Avalanche/Source/AvalanchePropertyAnimatorEditor/Private/Outliner/AvaPropertyAnimatorEditorOutlinerProxy.h b/Engine/Plugins/Experimental/Avalanche/Source/AvalanchePropertyAnimatorEditor/Private/Outliner/AvaPropertyAnimatorEditorOutlinerProxy.h
--- a/Engine/Plugins/Experimental/Avalanche/Source/AvalanchePropertyAnimatorEditor/Private/Outliner/AvaPropertyAnimatorEditorOutlinerProxy.h
+++ b/Engine/Plugins/Experimental/Avalanche/Source/AvalanchePropertyAnimatorEditor/Private/Outliner/AvaPropertyAnimatorEditorOutlinerProxy.h
@@ -37,7 +37,7 @@ protected:
 	void BindDelegates();
 	void UnbindDelegates();
 
-	void OnPropertyAnimatorUpdated(UPropertyAnimatorCoreBase* InAnimator);
+	void OnPropertyAnimatorUpdated(UPropertyAnimatorCoreComponent* InComponent, UPropertyAnimatorCoreBase* InAnimator);
 
 	FSlateIcon ItemIcon;
 };
diff --git a/Engine/Plugins/Experimental/PropertyAnimatorCore/Source/PropertyAnimatorCore/Private/Animators/PropertyAnimatorCoreBase.cpp b/Engine/Plugins/Experimental/PropertyAnimatorCore/Source/PropertyAnimatorCore/Private/Animators/PropertyAnimatorCoreBase.cpp
--- a/Engine/Plugins/Experimental/PropertyAnimatorCore/Source/PropertyAnimatorCore/Private/Animators/PropertyAnimatorCoreBase.cpp
+++ b/Engine/Plugins/Experimental/PropertyAnimatorCore/Source/PropertyAnimatorCore/Private/Animators/PropertyAnimatorCoreBase.cpp
@@ -13,7 +13,7 @@
 
 DEFINE_LOG_CATEGORY_STATIC(LogPropertyAnimatorCoreBase, Log, All);
 
-UPropertyAnimatorCoreBase::FOnAnimatorUpdated UPropertyAnimatorCoreBase::OnAnimatorCreatedDelegate;
+UPropertyAnimatorCoreBase::FOnAnimatorUpdated UPropertyAnimatorCoreBase::OnAnimatorAddedDelegate;
 UPropertyAnimatorCoreBase::FOnAnimatorUpdated UPropertyAnimatorCoreBase::OnAnimatorRemovedDelegate;
 UPropertyAnimatorCoreBase::FOnAnimatorUpdated UPropertyAnimatorCoreBase::OnAnimatorRenamedDelegate;
 UPropertyAnimatorCoreBase::FOnAnimatorPropertyUpdated UPropertyAnimatorCoreBase::OnAnimatorPropertyLinkedDelegate;
@@ -333,6 +333,13 @@ void UPropertyAnimatorCoreBase::OnAnimatorAdded()
 	{
 		SetTimeSourceName(AnimatorSettings->GetDefaultTimeSourceName());
 	}
+
+	UPropertyAnimatorCoreBase::OnAnimatorAddedDelegate.Broadcast(GetAnimatorComponent(), this);
+}
+
+void UPropertyAnimatorCoreBase::OnAnimatorRemoved()
+{
+	UPropertyAnimatorCoreBase::OnAnimatorRemovedDelegate.Broadcast(GetAnimatorComponent(), this);
 }
 
 void UPropertyAnimatorCoreBase::OnAnimatorEnabled()
@@ -608,7 +615,7 @@ void UPropertyAnimatorCoreBase::SetAnimatorDisplayName(FName InName)
 	AnimatorDisplayName = InName;
 	OnAnimatorDisplayNameChanged();
 
-	OnAnimatorRenamedDelegate.Broadcast(this);
+	OnAnimatorRenamedDelegate.Broadcast(GetAnimatorComponent(), this);
 }
 
 TSet<FPropertyAnimatorCoreData> UPropertyAnimatorCoreBase::GetLinkedProperties() const
diff --git a/Engine/Plugins/Experimental/PropertyAnimatorCore/Source/PropertyAnimatorCore/Private/Components/PropertyAnimatorCoreComponent.cpp b/Engine/Plugins/Experimental/PropertyAnimatorCore/Source/PropertyAnimatorCore/Private/Components/PropertyAnimatorCoreComponent.cpp
--- a/Engine/Plugins/Experimental/PropertyAnimatorCore/Source/PropertyAnimatorCore/Private/Components/PropertyAnimatorCoreComponent.cpp
+++ b/Engine/Plugins/Experimental/PropertyAnimatorCore/Source/PropertyAnimatorCore/Private/Components/PropertyAnimatorCoreComponent.cpp
@@ -95,15 +95,22 @@ void UPropertyAnimatorCoreComponent::OnAnimatorsChanged()
 	{
 		if (RemovedAnimator)
 		{
+#if WITH_EDITOR
+			RemovedAnimator->Modify();
+#endif
+
 			RemovedAnimator->SetAnimatorEnabled(false);
 			RemovedAnimator->OnAnimatorRemoved();
 		}
 	}
-
 	for (const TObjectPtr<UPropertyAnimatorCoreBase>& AddedAnimator : AddedAnimators)
 	{
 		if (AddedAnimator)
 		{
+#if WITH_EDITOR
+			AddedAnimator->Modify();
+#endif
+
 			AddedAnimator->SetAnimatorDisplayName(GetAnimatorName(AddedAnimator));
 			AddedAnimator->OnAnimatorAdded();
 			AddedAnimator->SetAnimatorEnabled(true);
@@ -244,6 +251,8 @@ void UPropertyAnimatorCoreComponent::OnComponentCreated()
 			}
 		}
 	}
+
+	UPropertyAnimatorCoreBase::OnAnimatorAddedDelegate.Broadcast(this, nullptr);
 }
 
 UPropertyAnimatorCoreComponent* UPropertyAnimatorCoreComponent::FindOrAdd(AActor* InActor)
@@ -327,6 +336,8 @@ void UPropertyAnimatorCoreComponent::OnComponentDestroyed(bool bInDestroyingHier
 	PropertyAnimators.Empty();
 
 	OnAnimatorsChanged();
+
+	UPropertyAnimatorCoreBase::OnAnimatorRemovedDelegate.Broadcast(this, nullptr);
 }
 
 void UPropertyAnimatorCoreComponent::TickComponent(float InDeltaTime, ELevelTick InTickType, FActorComponentTickFunction* InTickFunction)
diff --git a/Engine/Plugins/Experimental/PropertyAnimatorCore/Source/PropertyAnimatorCore/Private/Subsystems/PropertyAnimatorCoreSubsystem.cpp b/Engine/Plugins/Experimental/PropertyAnimatorCore/Source/PropertyAnimatorCore/Private/Subsystems/PropertyAnimatorCoreSubsystem.cpp
--- a/Engine/Plugins/Experimental/PropertyAnimatorCore/Source/PropertyAnimatorCore/Private/Subsystems/PropertyAnimatorCoreSubsystem.cpp
+++ b/Engine/Plugins/Experimental/PropertyAnimatorCore/Source/PropertyAnimatorCore/Private/Subsystems/PropertyAnimatorCoreSubsystem.cpp
@@ -371,8 +371,6 @@ TSet<UPropertyAnimatorCoreBase*> UPropertyAnimatorCoreSubsystem::CreateAnimators
 		}
 
 		NewAnimators.Add(NewActorAnimator);
-
-		UPropertyAnimatorCoreBase::OnAnimatorCreatedDelegate.Broadcast(NewActorAnimator);
 	}
 
 	return NewAnimators;
@@ -465,17 +463,49 @@ bool UPropertyAnimatorCoreSubsystem::RemoveAnimators(const TSet<UPropertyAnimato
 		Animator->Modify();
 #endif
 
-		const bool bAnimatorRemoved = Component->RemoveAnimator(Animator);
+		bResult &= Component->RemoveAnimator(Animator);
+	}
+
+	return bResult;
+}
+
+bool UPropertyAnimatorCoreSubsystem::RemoveAnimatorComponents(const TSet<UPropertyAnimatorCoreComponent*>& InComponents, bool bInTransact) const
+{
+	if (InComponents.IsEmpty())
+	{
+		return false;
+	}
+
+#if WITH_EDITOR
+	const FText TransactionText = LOCTEXT("RemoveAnimatorComponent", "Removing {0} animator component(s)");
+	const FText ComponentCount = FText::FromString(FString::FromInt(InComponents.Num()));
+
+	FScopedTransaction Transaction(FText::Format(TransactionText, ComponentCount), bInTransact);
+#endif
 
-		if (bAnimatorRemoved)
+	for (UPropertyAnimatorCoreComponent* Component : InComponents)
+	{
+		if (!IsValid(Component))
 		{
-			UPropertyAnimatorCoreBase::OnAnimatorRemovedDelegate.Broadcast(Animator);
+			continue;
+		}
+
+		AActor* OwningActor = Component->GetOwner();
+
+		if (!IsValid(OwningActor))
+		{
+			continue;
 		}
 
-		bResult &= bAnimatorRemoved;
+#if WITH_EDITOR
+		OwningActor->Modify();
+		Component->Modify();
+#endif
+
+		Component->DestroyComponent(/** PromoteChildren */false);
 	}
 
-	return bResult;
+	return true;
 }
 
 bool UPropertyAnimatorCoreSubsystem::ApplyAnimatorPreset(UPropertyAnimatorCoreBase* InAnimator, UPropertyAnimatorCorePresetBase* InPreset, bool bInTransact)
diff --git a/Engine/Plugins/Experimental/PropertyAnimatorCore/Source/PropertyAnimatorCore/Public/Animators/PropertyAnimatorCoreBase.h b/Engine/Plugins/Experimental/PropertyAnimatorCore/Source/PropertyAnimatorCore/Public/Animators/PropertyAnimatorCoreBase.h
--- a/Engine/Plugins/Experimental/PropertyAnimatorCore/Source/PropertyAnimatorCore/Public/Animators/PropertyAnimatorCoreBase.h
+++ b/Engine/Plugins/Experimental/PropertyAnimatorCore/Source/PropertyAnimatorCore/Public/Animators/PropertyAnimatorCoreBase.h
@@ -30,29 +30,38 @@ class UPropertyAnimatorCoreBase : public UObject
 	friend class UPropertyAnimatorCoreContext;
 
 public:
-	DECLARE_MULTICAST_DELEGATE_OneParam(FOnAnimatorUpdated, UPropertyAnimatorCoreBase* /* InAnimator */)
+	static constexpr const TCHAR* TimeElapsedParameterName = TEXT("TimeElapsed");
+	static constexpr const TCHAR* MagnitudeParameterName = TEXT("Magnitude");
+	static constexpr const TCHAR* FrequencyParameterName = TEXT("Frequency");
+	static constexpr const TCHAR* AlphaParameterName = TEXT("Alpha");
 
-	/** Called when a Animator is created */
-	PROPERTYANIMATORCORE_API static FOnAnimatorUpdated OnAnimatorCreatedDelegate;
+	DECLARE_MULTICAST_DELEGATE_TwoParams(FOnAnimatorUpdated, UPropertyAnimatorCoreComponent* /* InComponent */, UPropertyAnimatorCoreBase* /* InAnimator */)
+	DECLARE_MULTICAST_DELEGATE_TwoParams(FOnAnimatorPropertyUpdated, UPropertyAnimatorCoreBase* /* InAnimator */, const FPropertyAnimatorCoreData& /* InProperty */)
 
-	/** Called when a Animator is removed */
-	PROPERTYANIMATORCORE_API static FOnAnimatorUpdated OnAnimatorRemovedDelegate;
-
-	/** Called when a Animator is renamed */
-	PROPERTYANIMATORCORE_API static FOnAnimatorUpdated OnAnimatorRenamedDelegate;
+	static FOnAnimatorUpdated::RegistrationType& OnPropertyAnimatorAdded()
+	{
+		return OnAnimatorAddedDelegate;
+	}
 
-	DECLARE_MULTICAST_DELEGATE_TwoParams(FOnAnimatorPropertyUpdated, UPropertyAnimatorCoreBase* /* InAnimator */, const FPropertyAnimatorCoreData& /** InProperty */)
+	static FOnAnimatorUpdated::RegistrationType& OnPropertyAnimatorRemoved()
+	{
+		return OnAnimatorRemovedDelegate;
+	}
 
-	/** Called when a property is linked to a Animator */
-	PROPERTYANIMATORCORE_API static FOnAnimatorPropertyUpdated OnAnimatorPropertyLinkedDelegate;
+	static FOnAnimatorUpdated::RegistrationType& OnPropertyAnimatorRenamed()
+	{
+		return OnAnimatorRenamedDelegate;
+	}
 
-	/** Called when a property is unlinked to a Animator */
-	PROPERTYANIMATORCORE_API static FOnAnimatorPropertyUpdated OnAnimatorPropertyUnlinkedDelegate;
+	static FOnAnimatorPropertyUpdated::RegistrationType& OnPropertyAnimatorPropertyLinked()
+	{
+		return OnAnimatorPropertyLinkedDelegate;
+	}
 
-	static constexpr const TCHAR* TimeElapsedParameterName = TEXT("TimeElapsed");
-	static constexpr const TCHAR* MagnitudeParameterName = TEXT("Magnitude");
-	static constexpr const TCHAR* FrequencyParameterName = TEXT("Frequency");
-	static constexpr const TCHAR* AlphaParameterName = TEXT("Alpha");
+	static FOnAnimatorPropertyUpdated::RegistrationType& OnPropertyAnimatorPropertyUnlinked()
+	{
+		return OnAnimatorPropertyUnlinkedDelegate;
+	}
 
 #if WITH_EDITOR
 	PROPERTYANIMATORCORE_API static FName GetAnimatorEnabledPropertyName();
@@ -248,7 +257,7 @@ protected:
 	virtual void OnAnimatorDisplayNameChanged() {}
 
 	PROPERTYANIMATORCORE_API virtual void OnAnimatorAdded();
-	virtual void OnAnimatorRemoved() {}
+	PROPERTYANIMATORCORE_API virtual void OnAnimatorRemoved();
 
 	PROPERTYANIMATORCORE_API virtual void OnAnimatorEnabled();
 	PROPERTYANIMATORCORE_API virtual void OnAnimatorDisabled();
@@ -265,6 +274,21 @@ protected:
 	virtual void EvaluateProperties(FInstancedPropertyBag& InParameters) {}
 
 private:
+	/** Called when a Animator is created */
+	PROPERTYANIMATORCORE_API static FOnAnimatorUpdated OnAnimatorAddedDelegate;
+
+	/** Called when a Animator is removed */
+	PROPERTYANIMATORCORE_API static FOnAnimatorUpdated OnAnimatorRemovedDelegate;
+
+	/** Called when a Animator is renamed */
+	PROPERTYANIMATORCORE_API static FOnAnimatorUpdated OnAnimatorRenamedDelegate;
+
+	/** Called when a property is linked to a Animator */
+	PROPERTYANIMATORCORE_API static FOnAnimatorPropertyUpdated OnAnimatorPropertyLinkedDelegate;
+
+	/** Called when a property is unlinked to a Animator */
+	PROPERTYANIMATORCORE_API static FOnAnimatorPropertyUpdated OnAnimatorPropertyUnlinkedDelegate;
+
 	/** Restore modified properties to original state */
 	void RestoreProperties(bool bInForce = false);
 
diff --git a/Engine/Plugins/Experimental/PropertyAnimatorCore/Source/PropertyAnimatorCore/Public/Subsystems/PropertyAnimatorCoreSubsystem.h b/Engine/Plugins/Experimental/PropertyAnimatorCore/Source/PropertyAnimatorCore/Public/Subsystems/PropertyAnimatorCoreSubsystem.h
--- a/Engine/Plugins/Experimental/PropertyAnimatorCore/Source/PropertyAnimatorCore/Public/Subsystems/PropertyAnimatorCoreSubsystem.h
+++ b/Engine/Plugins/Experimental/PropertyAnimatorCore/Source/PropertyAnimatorCore/Public/Subsystems/PropertyAnimatorCoreSubsystem.h
@@ -155,6 +155,9 @@ public:
 	/** Removes animators from their owner */
 	PROPERTYANIMATORCORE_API bool RemoveAnimators(const TSet<UPropertyAnimatorCoreBase*>& InAnimators, bool bInTransact = false) const;
 
+	/** Removes animator components from their actor */
+	PROPERTYANIMATORCORE_API bool RemoveAnimatorComponents(const TSet<UPropertyAnimatorCoreComponent*>& InComponents, bool bInTransact = false) const;
+
 	/** Apply a preset on an existing animator */
 	PROPERTYANIMATORCORE_API bool ApplyAnimatorPreset(UPropertyAnimatorCoreBase* InAnimator, UPropertyAnimatorCorePresetBase* InPreset, bool bInTransact = false);
 
diff --git a/Engine/Plugins/Experimental/PropertyAnimatorCore/Source/PropertyAnimatorCoreEditor/Private/Customizations/PropertyAnimatorCoreEditorStackCustomization.cpp b/Engine/Plugins/Experimental/PropertyAnimatorCore/Source/PropertyAnimatorCoreEditor/Private/Customizations/PropertyAnimatorCoreEditorStackCustomization.cpp
--- a/Engine/Plugins/Experimental/PropertyAnimatorCore/Source/PropertyAnimatorCoreEditor/Private/Customizations/PropertyAnimatorCoreEditorStackCustomization.cpp
+++ b/Engine/Plugins/Experimental/PropertyAnimatorCore/Source/PropertyAnimatorCoreEditor/Private/Customizations/PropertyAnimatorCoreEditorStackCustomization.cpp
@@ -335,24 +335,24 @@ void UPropertyAnimatorCoreEditorStackCustomization::RemoveAnimatorAction(FOperat
 		return;
 	}
 
-	TSet<UPropertyAnimatorCoreBase*> Animators;
-
 	if (InItem->IsA<UPropertyAnimatorCoreBase>())
 	{
-		Animators.Append(InItem->GetAsArray<UPropertyAnimatorCoreBase>());
+		const TSet<UPropertyAnimatorCoreBase*> Animators(InItem->GetAsArray<UPropertyAnimatorCoreBase>());
+
+		if (!Subsystem->RemoveAnimators(Animators, /** Transact */true))
+		{
+			UE_LOG(LogPropertyAnimatorCoreEditorStackCustomization, Warning, TEXT("Could not remove %i animator(s)"), Animators.Num())
+		}
 	}
 	else if (InItem->IsA<UPropertyAnimatorCoreComponent>())
 	{
-		for (const UPropertyAnimatorCoreComponent* AnimatorComponent : InItem->GetAsArray<UPropertyAnimatorCoreComponent>())
+		const TSet<UPropertyAnimatorCoreComponent*> Components(InItem->GetAsArray<UPropertyAnimatorCoreComponent>());
+
+		if (!Subsystem->RemoveAnimatorComponents(Components, /** Transact */true))
 		{
-			Animators.Append(AnimatorComponent->GetAnimators());
+			UE_LOG(LogPropertyAnimatorCoreEditorStackCustomization, Warning, TEXT("Could not remove %i animator component(s)"), Components.Num())
 		}
 	}
-
-	if (!Subsystem->RemoveAnimators(Animators, /** Transact */true))
-	{
-		UE_LOG(LogPropertyAnimatorCoreEditorStackCustomization, Warning, TEXT("Could not remove %i animator(s)"), Animators.Num())
-	}
 }
 
 void UPropertyAnimatorCoreEditorStackCustomization::FillAddAnimatorMenuSection(UToolMenu* InToolMenu) const
