MotionDesign : PropertyAnimator
- Fixed wiggle animator cycle causing issues with randomness
- Fixed component time source not set on existing objects
- Filter time sources supported by numeric animator
- Renamed machine clock time source to System

diff --git a/Engine/Plugins/Experimental/PropertyAnimator/Source/PropertyAnimator/Private/Animators/PropertyAnimatorClock.cpp b/Engine/Plugins/Experimental/PropertyAnimator/Source/PropertyAnimator/Private/Animators/PropertyAnimatorClock.cpp
--- a/Engine/Plugins/Experimental/PropertyAnimator/Source/PropertyAnimator/Private/Animators/PropertyAnimatorClock.cpp
+++ b/Engine/Plugins/Experimental/PropertyAnimator/Source/PropertyAnimator/Private/Animators/PropertyAnimatorClock.cpp
@@ -117,7 +117,7 @@ void UPropertyAnimatorClock::EvaluateProperties(FInstancedPropertyBag& InParamet
 	const FDateTime DateTime(ElapsedTimeSpan > FTimespan::Zero() ? ElapsedTimeSpan.GetTicks() : 0);
 	const FString FormattedDateTime = FormatDateTime(DateTime, DisplayFormat);
 
-	EvaluateEachLinkedProperty<UPropertyAnimatorCoreContext>([this, FormattedDateTime](
+	EvaluateEachLinkedProperty([this, FormattedDateTime](
 		UPropertyAnimatorCoreContext* InContext
 		, const FPropertyAnimatorCoreData& InResolvedProperty
 		, FInstancedPropertyBag& InEvaluatedValues
diff --git a/Engine/Plugins/Experimental/PropertyAnimator/Source/PropertyAnimator/Private/Animators/PropertyAnimatorCounter.cpp b/Engine/Plugins/Experimental/PropertyAnimator/Source/PropertyAnimator/Private/Animators/PropertyAnimatorCounter.cpp
--- a/Engine/Plugins/Experimental/PropertyAnimator/Source/PropertyAnimator/Private/Animators/PropertyAnimatorCounter.cpp
+++ b/Engine/Plugins/Experimental/PropertyAnimator/Source/PropertyAnimator/Private/Animators/PropertyAnimatorCounter.cpp
@@ -280,7 +280,7 @@ void UPropertyAnimatorCounter::EvaluateProperties(FInstancedPropertyBag& InParam
 {
 	const double TimeElapsed = InParameters.GetValueDouble(TimeElapsedParameterName).GetValue();
 
-	EvaluateEachLinkedProperty<UPropertyAnimatorCoreContext>([this, TimeElapsed](
+	EvaluateEachLinkedProperty([this, TimeElapsed](
 		UPropertyAnimatorCoreContext* InContext
 		, const FPropertyAnimatorCoreData& InResolvedProperty
 		, FInstancedPropertyBag& InEvaluatedValues
diff --git a/Engine/Plugins/Experimental/PropertyAnimator/Source/PropertyAnimator/Private/Animators/PropertyAnimatorNumericBase.cpp b/Engine/Plugins/Experimental/PropertyAnimator/Source/PropertyAnimator/Private/Animators/PropertyAnimatorNumericBase.cpp
--- a/Engine/Plugins/Experimental/PropertyAnimator/Source/PropertyAnimator/Private/Animators/PropertyAnimatorNumericBase.cpp
+++ b/Engine/Plugins/Experimental/PropertyAnimator/Source/PropertyAnimator/Private/Animators/PropertyAnimatorNumericBase.cpp
@@ -8,6 +8,8 @@
 #include "Properties/Converters/PropertyAnimatorCoreConverterBase.h"
 #include "Properties/Handlers/PropertyAnimatorCoreHandlerBase.h"
 #include "Subsystems/PropertyAnimatorCoreSubsystem.h"
+#include "TimeSources/PropertyAnimatorCoreSystemTimeSource.h"
+#include "TimeSources/PropertyAnimatorCoreTimeSourceBase.h"
 
 #if WITH_EDITOR
 void UPropertyAnimatorNumericBase::PostEditChangeProperty(FPropertyChangedEvent& InPropertyChangedEvent)
@@ -49,6 +51,16 @@ void UPropertyAnimatorNumericBase::SetCycleDuration(float InCycleDuration)
 	OnCycleDurationChanged();
 }
 
+void UPropertyAnimatorNumericBase::SetCycleGapDuration(float InCycleGap)
+{
+	CycleGapDuration = FMath::Max(0, InCycleGap);
+}
+
+void UPropertyAnimatorNumericBase::SetCycleRate(float InFrequency)
+{
+	CycleRate = FMath::Max(0, InFrequency);
+}
+
 void UPropertyAnimatorNumericBase::SetCycleMode(EPropertyAnimatorCycleMode InMode)
 {
 	if (CycleMode == InMode)
@@ -156,7 +168,7 @@ void UPropertyAnimatorNumericBase::EvaluateProperties(FInstancedPropertyBag& InP
 	double TimeElapsed = InParameters.GetValueDouble(TimeElapsedParameterName).GetValue();
 	RandomStream = FRandomStream(Seed);
 
-	EvaluateEachLinkedProperty<UPropertyAnimatorCoreContext>([this, &TimeElapsed, &AnimatorMagnitude, &InParameters](
+	EvaluateEachLinkedProperty([this, &TimeElapsed, &AnimatorMagnitude, &InParameters](
 		UPropertyAnimatorCoreContext* InOptions
 		, const FPropertyAnimatorCoreData& InResolvedProperty
 		, FInstancedPropertyBag& InEvaluatedValues
@@ -168,6 +180,7 @@ void UPropertyAnimatorNumericBase::EvaluateProperties(FInstancedPropertyBag& InP
 		const double AbsTimeOffset = FMath::Abs(TimeOffset);
 		const double MaxTimeOffset = InRangeMax * AbsTimeOffset;
 		double PropertyTimeElapsed = TimeElapsed - MaxTimeOffset + RandomTimeOffset;
+		double Frequency = CycleDuration != 0.f ? 1.f / CycleDuration : 0.f;
 
 		if (TimeOffset >= 0)
 		{
@@ -213,6 +226,10 @@ void UPropertyAnimatorNumericBase::EvaluateProperties(FInstancedPropertyBag& InP
 				PropertyTimeElapsed = FMath::Fmod(PropertyTimeElapsed, CycleDuration);
 			}
 		}
+		else if (CycleMode == EPropertyAnimatorCycleMode::None)
+		{
+			Frequency = CycleRate;
+		}
 
 		if (Magnitude != 0
 			&& CycleDuration > 0
@@ -220,7 +237,7 @@ void UPropertyAnimatorNumericBase::EvaluateProperties(FInstancedPropertyBag& InP
 		{
 			// Frequency
 			InParameters.AddProperty(FrequencyParameterName, EPropertyBagPropertyType::Float);
-			InParameters.SetValueFloat(FrequencyParameterName, 1.f / CycleDuration);
+			InParameters.SetValueFloat(FrequencyParameterName, Frequency);
 
 			// Time Elapsed
 			InParameters.SetValueDouble(TimeElapsedParameterName, PropertyTimeElapsed);
@@ -254,3 +271,8 @@ void UPropertyAnimatorNumericBase::OnPropertyLinked(UPropertyAnimatorCoreContext
 		}
 	}
 }
+
+bool UPropertyAnimatorNumericBase::IsTimeSourceSupported(UPropertyAnimatorCoreTimeSourceBase* InTimeSource) const
+{
+	return !InTimeSource->IsA<UPropertyAnimatorCoreSystemTimeSource>();
+}
diff --git a/Engine/Plugins/Experimental/PropertyAnimator/Source/PropertyAnimator/Private/Animators/PropertyAnimatorWiggle.cpp b/Engine/Plugins/Experimental/PropertyAnimator/Source/PropertyAnimator/Private/Animators/PropertyAnimatorWiggle.cpp
--- a/Engine/Plugins/Experimental/PropertyAnimator/Source/PropertyAnimator/Private/Animators/PropertyAnimatorWiggle.cpp
+++ b/Engine/Plugins/Experimental/PropertyAnimator/Source/PropertyAnimator/Private/Animators/PropertyAnimatorWiggle.cpp
@@ -14,6 +14,7 @@ UPropertyAnimatorWiggle::UPropertyAnimatorWiggle()
 
 	bRandomTimeOffset = true;
 	Seed = SeedIncrement++;
+	CycleMode = EPropertyAnimatorCycleMode::None;
 }
 
 bool UPropertyAnimatorWiggle::EvaluateProperty(const FPropertyAnimatorCoreData& InPropertyData, UPropertyAnimatorCoreContext* InContext, FInstancedPropertyBag& InParameters, FInstancedPropertyBag& OutEvaluationResult) const
diff --git a/Engine/Plugins/Experimental/PropertyAnimator/Source/PropertyAnimator/Public/Animators/PropertyAnimatorCurve.h b/Engine/Plugins/Experimental/PropertyAnimator/Source/PropertyAnimator/Public/Animators/PropertyAnimatorCurve.h
--- a/Engine/Plugins/Experimental/PropertyAnimator/Source/PropertyAnimator/Public/Animators/PropertyAnimatorCurve.h
+++ b/Engine/Plugins/Experimental/PropertyAnimator/Source/PropertyAnimator/Public/Animators/PropertyAnimatorCurve.h
@@ -83,21 +83,23 @@ protected:
 	void OnEaseOutChanged();
 	virtual void OnCycleDurationChanged() override;
 
-	/** The wave curve to sample for the animation */
-	UPROPERTY(EditInstanceOnly, Setter, Getter, Category="Animator")
+	/** The wave curve to sample for the animation, base curve */
+	UPROPERTY(EditInstanceOnly, Setter, Getter, DisplayName="Base Curve", Category="Animator")
 	TObjectPtr<UPropertyAnimatorWaveCurve> WaveCurve;
 
-	UPROPERTY(EditInstanceOnly, Setter="SetEaseInEnabled", Getter="GetEaseInEnabled", Category="Animator", meta=(InlineEditConditionToggle))
+	/** Use ease in effect */
+	UPROPERTY(EditInstanceOnly, Setter="SetEaseInEnabled", Getter="GetEaseInEnabled", DisplayName="InEnabled", Category="Animator", meta=(InlineEditConditionToggle))
 	bool bEaseInEnabled = false;
 
 	/** Ease in for this effect */
-	UPROPERTY(EditInstanceOnly, Setter, Getter, Category="Animator", meta=(EditCondition="bEaseInEnabled"))
+	UPROPERTY(EditInstanceOnly, Setter, Getter, DisplayName="In", Category="Animator", meta=(EditCondition="bEaseInEnabled"))
 	FPropertyAnimatorCurveEasing EaseIn;
 
-	UPROPERTY(EditInstanceOnly, Setter="SetEaseOutEnabled", Getter="GetEaseOutEnabled", Category="Animator", meta=(InlineEditConditionToggle))
+	/** Use ease out effect */
+	UPROPERTY(EditInstanceOnly, Setter="SetEaseOutEnabled", Getter="GetEaseOutEnabled", DisplayName="OutEnabled", Category="Animator", meta=(InlineEditConditionToggle))
 	bool bEaseOutEnabled = false;
 
 	/** Ease out for this effect */
-	UPROPERTY(EditInstanceOnly, Setter, Getter, Category="Animator", meta=(EditCondition="bEaseOutEnabled"))
+	UPROPERTY(EditInstanceOnly, Setter, Getter, Category="Animator", DisplayName="Out", meta=(EditCondition="bEaseOutEnabled"))
     FPropertyAnimatorCurveEasing EaseOut;
 };
\ No newline at end of file
diff --git a/Engine/Plugins/Experimental/PropertyAnimator/Source/PropertyAnimator/Public/Animators/PropertyAnimatorNumericBase.h b/Engine/Plugins/Experimental/PropertyAnimator/Source/PropertyAnimator/Public/Animators/PropertyAnimatorNumericBase.h
--- a/Engine/Plugins/Experimental/PropertyAnimator/Source/PropertyAnimator/Public/Animators/PropertyAnimatorNumericBase.h
+++ b/Engine/Plugins/Experimental/PropertyAnimator/Source/PropertyAnimator/Public/Animators/PropertyAnimatorNumericBase.h
@@ -8,6 +8,8 @@
 UENUM(BlueprintType)
 enum class EPropertyAnimatorCycleMode : uint8
 {
+	/** Don't cycle time at all */
+	None	UMETA(Hidden),
 	/** Cycle only once then stop */
 	DoOnce,
 	/** Cycle and repeat once we reached the end */
@@ -33,16 +35,28 @@ public:
 		return Magnitude;
 	}
 
+	PROPERTYANIMATOR_API void SetCycleMode(EPropertyAnimatorCycleMode InMode);
+	EPropertyAnimatorCycleMode GetCycleMode() const
+	{
+		return CycleMode;
+	}
+
 	PROPERTYANIMATOR_API void SetCycleDuration(float InCycleDuration);
 	float GetCycleDuration() const
 	{
 		return CycleDuration;
 	}
 
-	PROPERTYANIMATOR_API void SetCycleMode(EPropertyAnimatorCycleMode InMode);
-	EPropertyAnimatorCycleMode GetCycleMode() const
+	PROPERTYANIMATOR_API void SetCycleGapDuration(float InCycleGap);
+	float GetCycleGapDuration() const
 	{
-		return CycleMode;
+		return CycleGapDuration;
+	}
+
+	PROPERTYANIMATOR_API void SetCycleRate(float InFrequency);
+	float GetCycleRate() const
+	{
+		return CycleRate;
 	}
 
 	PROPERTYANIMATOR_API void SetTimeOffset(double InOffset);
@@ -81,6 +95,7 @@ protected:
 	virtual EPropertyAnimatorPropertySupport IsPropertySupported(const FPropertyAnimatorCoreData& InPropertyData) const override;
 	virtual void EvaluateProperties(FInstancedPropertyBag& InParameters) override;
 	virtual void OnPropertyLinked(UPropertyAnimatorCoreContext* InLinkedProperty, EPropertyAnimatorPropertySupport InSupport) override;
+	virtual bool IsTimeSourceSupported(UPropertyAnimatorCoreTimeSourceBase* InTimeSource) const override;
 	//~ End UPropertyAnimatorCoreBase
 
 	/** Evaluate and return float value for a property */
@@ -93,18 +108,22 @@ protected:
 	UPROPERTY(EditInstanceOnly, Setter, Getter, Category="Animator", meta=(ClampMin="0"))
 	float Magnitude = 1.f;
 
-	/** Duration of one cycle for the effect = period of the effect */
-	UPROPERTY(EditInstanceOnly, Setter, Getter, Category="Animator", meta=(ClampMin="0", Units=Seconds))
-	float CycleDuration = 1.f;
-
 	/** Cycle mode for the effect */
-	UPROPERTY(EditInstanceOnly, Setter, Getter, Category="Animator")
+	UPROPERTY(EditInstanceOnly, Setter, Getter, Category="Animator", meta=(EditCondition="CycleMode != EPropertyAnimatorCycleMode::None", EditConditionHides))
 	EPropertyAnimatorCycleMode CycleMode = EPropertyAnimatorCycleMode::Loop;
 
+	/** Duration of one cycle for the effect = period of the effect */
+	UPROPERTY(EditInstanceOnly, Setter, Getter, Category="Animator", meta=(ClampMin="0", Units=Seconds, EditCondition="CycleMode != EPropertyAnimatorCycleMode::None", EditConditionHides))
+	float CycleDuration = 1.f;
+
 	/** Time gap between each cycle */
-	UPROPERTY(EditInstanceOnly, Category="Animator", meta=(ClampMin="0", Units=Seconds, EditCondition="CycleMode != EPropertyAnimatorCycleMode::DoOnce", EditConditionHides))
+	UPROPERTY(EditInstanceOnly, Setter, Getter, Category="Animator", meta=(ClampMin="0", Units=Seconds, EditCondition="CycleMode != EPropertyAnimatorCycleMode::DoOnce && CycleMode != EPropertyAnimatorCycleMode::None", EditConditionHides))
 	float CycleGapDuration = 0.f;
 
+	/** Frequency for the effect, higher values will give you faster movements */
+	UPROPERTY(EditInstanceOnly, Setter, Getter, Category="Animator", DisplayName="Frequency", meta=(ClampMin="0", Units=Hz, EditCondition="CycleMode == EPropertyAnimatorCycleMode::None", EditConditionHides))
+	float CycleRate = 1.f;
+
 	/** Use random time offset to add variation in animation */
 	UPROPERTY(EditInstanceOnly, Setter="SetRandomTimeOffset", Getter="GetRandomTimeOffset", Category="Animator", meta=(InlineEditConditionToggle))
 	bool bRandomTimeOffset = false;
diff --git a/Engine/Plugins/Experimental/PropertyAnimator/Source/PropertyAnimator/Public/Animators/PropertyAnimatorWiggle.h b/Engine/Plugins/Experimental/PropertyAnimator/Source/PropertyAnimator/Public/Animators/PropertyAnimatorWiggle.h
--- a/Engine/Plugins/Experimental/PropertyAnimator/Source/PropertyAnimator/Public/Animators/PropertyAnimatorWiggle.h
+++ b/Engine/Plugins/Experimental/PropertyAnimator/Source/PropertyAnimator/Public/Animators/PropertyAnimatorWiggle.h
@@ -6,7 +6,7 @@
 #include "PropertyAnimatorWiggle.generated.h"
 
 /**
- * Applies an additive random wiggle movement with various options on supported float properties
+ * Applies a random wiggle movement with various options on supported numeric properties
  */
 UCLASS(MinimalAPI, AutoExpandCategories=("Animator"))
 class UPropertyAnimatorWiggle : public UPropertyAnimatorNumericBase
diff --git a/Engine/Plugins/Experimental/PropertyAnimatorCore/Config/BasePropertyAnimatorCore.ini b/Engine/Plugins/Experimental/PropertyAnimatorCore/Config/BasePropertyAnimatorCore.ini
--- a/Engine/Plugins/Experimental/PropertyAnimatorCore/Config/BasePropertyAnimatorCore.ini
+++ b/Engine/Plugins/Experimental/PropertyAnimatorCore/Config/BasePropertyAnimatorCore.ini
@@ -12,10 +12,13 @@
 +ClassRedirects=(OldName="/Script/PropertyAnimator.PropertyAnimatorTimeSourceBase",NewName="/Script/PropertyAnimatorCore.PropertyAnimatorCoreTimeSourceBase")
 +ClassRedirects=(OldName="/Script/PropertyAnimator.PropertyAnimatorWorldTimeSource",NewName="/Script/PropertyAnimatorCore.PropertyAnimatorCoreWorldTimeSource")
 +ClassRedirects=(OldName="/Script/AvalanchePropertyAnimator.AvaPropertyAnimatorSequenceTimeSource",NewName="/Script/PropertyAnimatorCore.PropertyAnimatorCoreSequencerTimeSource")
++ClassRedirects=(OldName="/Script/PropertyAnimatorCore.PropertyAnimatorCoreMachineClockTimeSource",NewName="/Script/PropertyAnimatorCore.PropertyAnimatorCoreSystemTimeSource")
 
 +StructRedirects=(OldName="/Script/CustomPropertyControlEditor.CustomPropertyControlData",NewName="/Script/PropertyAnimatorCore.PropertyAnimatorCoreData")
 
 +PropertyRedirects=(OldName="/Script/PropertyAnimatorCore.PropertyAnimatorCoreBase.ControllerDisplayName",NewName="/Script/PropertyAnimatorCore.PropertyAnimatorCoreBase.AnimatorDisplayName")
 +PropertyRedirects=(OldName="/Script/PropertyAnimatorCore.PropertyAnimatorCoreComponent.Controllers",NewName="/Script/PropertyAnimatorCore.PropertyAnimatorCoreComponent.Animators")
 +PropertyRedirects=(OldName="/Script/PropertyAnimatorCore.PropertyAnimatorCoreBase.bEnabled",NewName="/Script/PropertyAnimatorCore.PropertyAnimatorCoreBase.bAnimatorEnabled")
-+PropertyRedirects=(OldName="/Script/PropertyAnimatorCore.PropertyAnimatorCoreComponent.bEnabled",NewName="/Script/PropertyAnimatorCore.PropertyAnimatorCoreComponent.bAnimatorsEnabled")
\ No newline at end of file
++PropertyRedirects=(OldName="/Script/PropertyAnimatorCore.PropertyAnimatorCoreComponent.bEnabled",NewName="/Script/PropertyAnimatorCore.PropertyAnimatorCoreComponent.bAnimatorsEnabled")
+
++EnumRedirects=(OldName="/Script/PropertyAnimatorCore.EPropertyAnimatorCoreMachineClockMode",NewName="/Script/PropertyAnimatorCore.EPropertyAnimatorCoreSystemMode")
\ No newline at end of file
diff --git a/Engine/Plugins/Experimental/PropertyAnimatorCore/Source/PropertyAnimatorCore/Private/Animators/PropertyAnimatorCoreBase.cpp b/Engine/Plugins/Experimental/PropertyAnimatorCore/Source/PropertyAnimatorCore/Private/Animators/PropertyAnimatorCoreBase.cpp
--- a/Engine/Plugins/Experimental/PropertyAnimatorCore/Source/PropertyAnimatorCore/Private/Animators/PropertyAnimatorCoreBase.cpp
+++ b/Engine/Plugins/Experimental/PropertyAnimatorCore/Source/PropertyAnimatorCore/Private/Animators/PropertyAnimatorCoreBase.cpp
@@ -544,7 +544,13 @@ TArray<FName> UPropertyAnimatorCoreBase::GetTimeSourceNames() const
 
 	if (const UPropertyAnimatorCoreSubsystem* AnimatorSubsystem = UPropertyAnimatorCoreSubsystem::Get())
 	{
-		TimeSourceNames = AnimatorSubsystem->GetTimeSourceNames();
+		for (UPropertyAnimatorCoreTimeSourceBase* TimeSource : AnimatorSubsystem->GetTimeSources())
+		{
+			if (TimeSource && IsTimeSourceSupported(TimeSource))
+			{
+				TimeSourceNames.Add(TimeSource->GetTimeSourceName());
+			}
+		}
 	}
 
 	return TimeSourceNames;
diff --git a/Engine/Plugins/Experimental/PropertyAnimatorCore/Source/PropertyAnimatorCore/Private/Subsystems/PropertyAnimatorCoreSubsystem.cpp b/Engine/Plugins/Experimental/PropertyAnimatorCore/Source/PropertyAnimatorCore/Private/Subsystems/PropertyAnimatorCoreSubsystem.cpp
--- a/Engine/Plugins/Experimental/PropertyAnimatorCore/Source/PropertyAnimatorCore/Private/Subsystems/PropertyAnimatorCoreSubsystem.cpp
+++ b/Engine/Plugins/Experimental/PropertyAnimatorCore/Source/PropertyAnimatorCore/Private/Subsystems/PropertyAnimatorCoreSubsystem.cpp
@@ -897,6 +897,26 @@ TArray<FName> UPropertyAnimatorCoreSubsystem::GetTimeSourceNames() const
 	return TimeSourceNames;
 }
 
+TArray<UPropertyAnimatorCoreTimeSourceBase*> UPropertyAnimatorCoreSubsystem::GetTimeSources() const
+{
+	TArray<UPropertyAnimatorCoreTimeSourceBase*> TimeSources;
+
+	Algo::TransformIf(
+		TimeSourcesWeak
+		, TimeSources
+		, [](const TWeakObjectPtr<UPropertyAnimatorCoreTimeSourceBase>& InTimeSourceWeak)
+		{
+			return InTimeSourceWeak.IsValid();
+		}
+		, [](const TWeakObjectPtr<UPropertyAnimatorCoreTimeSourceBase>& InTimeSourceWeak)
+		{
+			return InTimeSourceWeak.Get();
+		}
+	);
+
+	return TimeSources;
+}
+
 UPropertyAnimatorCoreTimeSourceBase* UPropertyAnimatorCoreSubsystem::GetTimeSource(FName InTimeSourceName) const
 {
 	if (InTimeSourceName.IsNone())
diff --git a/Engine/Plugins/Experimental/PropertyAnimatorCore/Source/PropertyAnimatorCore/Private/TimeSources/PropertyAnimatorCoreMachineClockTimeSource.cpp b/Engine/Plugins/Experimental/PropertyAnimatorCore/Source/PropertyAnimatorCore/Private/TimeSources/PropertyAnimatorCoreMachineClockTimeSource.cpp
deleted file mode 100644
diff --git a/Engine/Plugins/Experimental/PropertyAnimatorCore/Source/PropertyAnimatorCore/Private/TimeSources/PropertyAnimatorCoreSystemTimeSource.cpp b/Engine/Plugins/Experimental/PropertyAnimatorCore/Source/PropertyAnimatorCore/Private/TimeSources/PropertyAnimatorCoreSystemTimeSource.cpp
new file mode 100644
--- /dev/null
+++ b/Engine/Plugins/Experimental/PropertyAnimatorCore/Source/PropertyAnimatorCore/Private/TimeSources/PropertyAnimatorCoreSystemTimeSource.cpp
@@ -0,0 +1,139 @@
+// Copyright Epic Games, Inc. All Rights Reserved.
+
+#include "TimeSources/PropertyAnimatorCoreSystemTimeSource.h"
+
+#include "Internationalization/Regex.h"
+#include "Misc/DateTime.h"
+
+double UPropertyAnimatorCoreSystemTimeSource::GetTimeElapsed()
+{
+	double TimeElapsedSeconds = 0;
+
+	switch(Mode)
+	{
+	case EPropertyAnimatorCoreSystemMode::LocalTime:
+		{
+			TimeElapsedSeconds = (FDateTime::Now() - FDateTime::MinValue()).GetTotalSeconds();
+		}
+		break;
+	case EPropertyAnimatorCoreSystemMode::UtcTime:
+		{
+			TimeElapsedSeconds = (FDateTime::UtcNow() - FDateTime::MinValue()).GetTotalSeconds();
+		}
+		break;
+	case EPropertyAnimatorCoreSystemMode::Countdown:
+		{
+			TimeElapsedSeconds = (CountdownTimeSpan - (FDateTime::Now() - ActivationTime)).GetTotalSeconds();
+		}
+		break;
+	case EPropertyAnimatorCoreSystemMode::Stopwatch:
+		{
+			TimeElapsedSeconds = (FDateTime::Now() - ActivationTime).GetTotalSeconds();
+		}
+		break;
+	}
+
+	return TimeElapsedSeconds;
+}
+
+bool UPropertyAnimatorCoreSystemTimeSource::IsTimeSourceReady() const
+{
+	return true;
+}
+
+void UPropertyAnimatorCoreSystemTimeSource::OnTimeSourceActive()
+{
+	Super::OnTimeSourceActive();
+
+	ActivationTime = FDateTime::Now();
+	OnModeChanged();
+}
+
+void UPropertyAnimatorCoreSystemTimeSource::SetMode(EPropertyAnimatorCoreSystemMode InMode)
+{
+	if (Mode == InMode)
+	{
+		return;
+	}
+
+	Mode = InMode;
+	OnModeChanged();
+}
+
+void UPropertyAnimatorCoreSystemTimeSource::SetCountdownDuration(const FTimespan& InTimeSpan)
+{
+	if (InTimeSpan == CountdownTimeSpan)
+	{
+		return;
+	}
+
+	SetCountdownDuration(InTimeSpan.ToString(TEXT("%h:%m:%s")));
+}
+
+void UPropertyAnimatorCoreSystemTimeSource::SetCountdownDuration(const FString& InDuration)
+{
+	if (CountdownDuration == InDuration)
+	{
+		return;
+	}
+
+	CountdownDuration = InDuration;
+	OnModeChanged();
+}
+
+#if WITH_EDITOR
+void UPropertyAnimatorCoreSystemTimeSource::PostEditChangeProperty(FPropertyChangedEvent& InPropertyChangedEvent)
+{
+	Super::PostEditChangeProperty(InPropertyChangedEvent);
+
+	const FName MemberName = InPropertyChangedEvent.GetMemberPropertyName();
+
+	if (MemberName == GET_MEMBER_NAME_CHECKED(UPropertyAnimatorCoreSystemTimeSource, Mode)
+		|| MemberName == GET_MEMBER_NAME_CHECKED(UPropertyAnimatorCoreSystemTimeSource, CountdownDuration))
+	{
+		OnModeChanged();
+	}
+}
+#endif
+
+void UPropertyAnimatorCoreSystemTimeSource::OnModeChanged()
+{
+	if (Mode == EPropertyAnimatorCoreSystemMode::Countdown)
+	{
+		CountdownTimeSpan = ParseTime(CountdownDuration);
+	}
+}
+
+FTimespan UPropertyAnimatorCoreSystemTimeSource::ParseTime(const FString& InFormat)
+{
+	// Regex patterns for different formats
+	static const FRegexPattern HHMMSSPattern(TEXT("^(?:(\\d{2}):)?(\\d{2}):(\\d{2})$")); // 01:00 00:01:00
+	static const FRegexPattern CombinedPattern(TEXT("(?:(\\d+)h)? ?(?:(\\d+)m)? ?(?:(\\d+)s)?")); // 1h 1m 1s
+
+	FRegexMatcher HHMMSSMatcher(HHMMSSPattern, InFormat);
+	FRegexMatcher CombinedMatcher(CombinedPattern, InFormat);
+
+	FTimespan ParsedTimeSpan = FTimespan::Zero();
+
+	if (InFormat.IsNumeric())
+	{
+		const int32 Seconds = FCString::Atoi(*InFormat);
+		ParsedTimeSpan = FTimespan::FromSeconds(Seconds);
+	}
+	else if (HHMMSSMatcher.FindNext())
+	{
+		const int32 Hours = HHMMSSMatcher.GetCaptureGroup(1).IsEmpty() ? 0 : FCString::Atoi(*HHMMSSMatcher.GetCaptureGroup(1));
+		const int32 Minutes = FCString::Atoi(*HHMMSSMatcher.GetCaptureGroup(2));
+		const int32 Seconds = FCString::Atoi(*HHMMSSMatcher.GetCaptureGroup(3));
+		ParsedTimeSpan = FTimespan::FromHours(Hours) + FTimespan::FromMinutes(Minutes) + FTimespan::FromSeconds(Seconds);
+	}
+	else if (CombinedMatcher.FindNext())
+	{
+		const int32 Hours = CombinedMatcher.GetCaptureGroup(1).IsEmpty() ? 0 : FCString::Atoi(*CombinedMatcher.GetCaptureGroup(1));
+		const int32 Minutes = CombinedMatcher.GetCaptureGroup(2).IsEmpty() ? 0 : FCString::Atoi(*CombinedMatcher.GetCaptureGroup(2));
+		const int32 Seconds = CombinedMatcher.GetCaptureGroup(3).IsEmpty() ? 0 : FCString::Atoi(*CombinedMatcher.GetCaptureGroup(3));
+		ParsedTimeSpan = FTimespan::FromHours(Hours) + FTimespan::FromMinutes(Minutes) + FTimespan::FromSeconds(Seconds);
+	}
+
+	return ParsedTimeSpan;
+}
diff --git a/Engine/Plugins/Experimental/PropertyAnimatorCore/Source/PropertyAnimatorCore/Public/Animators/PropertyAnimatorCoreBase.h b/Engine/Plugins/Experimental/PropertyAnimatorCore/Source/PropertyAnimatorCore/Public/Animators/PropertyAnimatorCoreBase.h
--- a/Engine/Plugins/Experimental/PropertyAnimatorCore/Source/PropertyAnimatorCore/Public/Animators/PropertyAnimatorCoreBase.h
+++ b/Engine/Plugins/Experimental/PropertyAnimatorCore/Source/PropertyAnimatorCore/Public/Animators/PropertyAnimatorCoreBase.h
@@ -146,6 +146,12 @@ public:
 		return EPropertyAnimatorPropertySupport::None;
 	}
 
+	/** Checks whether a time source is supported on this animator */
+	virtual bool IsTimeSourceSupported(UPropertyAnimatorCoreTimeSourceBase* InTimeSource) const
+	{
+		return true;
+	}
+
 	/** Get the context for the linked property */
 	PROPERTYANIMATORCORE_API UPropertyAnimatorCoreContext* GetLinkedPropertyContext(const FPropertyAnimatorCoreData& InProperty) const;
 
@@ -212,11 +218,9 @@ protected:
 	void UpdateAnimatorDisplayName();
 
 	/** Used to evaluate linked properties, assign the result in the property bag and return true on success to update property value */
-	template<typename InContextClass
-		UE_REQUIRES(TIsDerivedFrom<InContextClass, UPropertyAnimatorCoreContext>::Value)>
 	void EvaluateEachLinkedProperty(
 		TFunctionRef<bool(
-			InContextClass* /** InPropertyContext */
+			UPropertyAnimatorCoreContext* /** InPropertyContext */
 			, const FPropertyAnimatorCoreData& /** InResolvedProperty */
 			, FInstancedPropertyBag& /** OutEvaluation */
 			, int32 InRangeIndex
@@ -227,28 +231,27 @@ protected:
 
 		for (const TObjectPtr<UPropertyAnimatorCoreContext>& LinkedProperty : LinkedProperties)
 		{
-			if (InContextClass* PropertyContext = Cast<InContextClass>(LinkedProperty.Get()))
+			UPropertyAnimatorCoreContext* PropertyContext = LinkedProperty.Get();
+
+			if (!PropertyContext || !PropertyContext->IsAnimated())
 			{
-				if (!PropertyContext->IsAnimated())
+				continue;
+			}
+
+			const TArray<FPropertyAnimatorCoreData> ResolvedProperties = PropertyContext->ResolveProperty(/** ForEvaluation */true);
+
+			for (int32 Index = 0; Index < ResolvedProperties.Num(); Index++)
+			{
+				const FPropertyAnimatorCoreData& ResolvedPropertyData = ResolvedProperties[Index];
+
+				if (!ResolvedPropertyData.IsResolved())
 				{
 					continue;
 				}
 
-				const TArray<FPropertyAnimatorCoreData> ResolvedProperties = PropertyContext->ResolveProperty(true);
-
-				for (int32 Index = 0; Index < ResolvedProperties.Num(); Index++)
+				if (InFunction(PropertyContext, ResolvedPropertyData, EvaluatedPropertyValues, Index, ResolvedProperties.Num() - 1))
 				{
-					const FPropertyAnimatorCoreData& ResolvedPropertyData = ResolvedProperties[Index];
-
-					if (!ResolvedPropertyData.IsResolved())
-					{
-						continue;
-					}
-
-					if (InFunction(PropertyContext, ResolvedPropertyData, EvaluatedPropertyValues, Index, ResolvedProperties.Num() - 1))
-					{
-						PropertyContext->CommitEvaluationResult(ResolvedPropertyData, EvaluatedPropertyValues);
-					}
+					PropertyContext->CommitEvaluationResult(ResolvedPropertyData, EvaluatedPropertyValues);
 				}
 			}
 		}
@@ -329,7 +332,7 @@ private:
 	TArray<TObjectPtr<UPropertyAnimatorCoreContext>> LinkedProperties;
 
 	/** Groups for properties linked to this Animator */
-	UPROPERTY(EditInstanceOnly, NoClear, Export, Instanced, Category="Animator")
+	UPROPERTY()
 	TArray<TObjectPtr<UPropertyAnimatorCoreGroupBase>> PropertyGroups;
 
 	/** Use the global time source or override it on this animator */
diff --git a/Engine/Plugins/Experimental/PropertyAnimatorCore/Source/PropertyAnimatorCore/Public/Subsystems/PropertyAnimatorCoreSubsystem.h b/Engine/Plugins/Experimental/PropertyAnimatorCore/Source/PropertyAnimatorCore/Public/Subsystems/PropertyAnimatorCoreSubsystem.h
--- a/Engine/Plugins/Experimental/PropertyAnimatorCore/Source/PropertyAnimatorCore/Public/Subsystems/PropertyAnimatorCoreSubsystem.h
+++ b/Engine/Plugins/Experimental/PropertyAnimatorCore/Source/PropertyAnimatorCore/Public/Subsystems/PropertyAnimatorCoreSubsystem.h
@@ -93,9 +93,12 @@ public:
 	/** Check time source class is registered */
 	PROPERTYANIMATORCORE_API bool IsTimeSourceClassRegistered(UClass* InTimeSourceClass) const;
 
-	/** Get all time sources available */
+	/** Get all time source names available */
 	TArray<FName> GetTimeSourceNames() const;
 
+	/** Get all time sources available */
+	TArray<UPropertyAnimatorCoreTimeSourceBase*> GetTimeSources() const;
+
 	/** Get a registered time source using its name */
 	UPropertyAnimatorCoreTimeSourceBase* GetTimeSource(FName InTimeSourceName) const;
 
diff --git a/Engine/Plugins/Experimental/PropertyAnimatorCore/Source/PropertyAnimatorCore/Public/TimeSources/PropertyAnimatorCoreMachineClockTimeSource.h b/Engine/Plugins/Experimental/PropertyAnimatorCore/Source/PropertyAnimatorCore/Public/TimeSources/PropertyAnimatorCoreMachineClockTimeSource.h
deleted file mode 100644
diff --git a/Engine/Plugins/Experimental/PropertyAnimatorCore/Source/PropertyAnimatorCore/Public/TimeSources/PropertyAnimatorCoreSystemTimeSource.h b/Engine/Plugins/Experimental/PropertyAnimatorCore/Source/PropertyAnimatorCore/Public/TimeSources/PropertyAnimatorCoreSystemTimeSource.h
new file mode 100644
--- /dev/null
+++ b/Engine/Plugins/Experimental/PropertyAnimatorCore/Source/PropertyAnimatorCore/Public/TimeSources/PropertyAnimatorCoreSystemTimeSource.h
@@ -0,0 +1,92 @@
+// Copyright Epic Games, Inc. All Rights Reserved.
+
+#pragma once
+
+#include "Misc/DateTime.h"
+#include "Misc/Timespan.h"
+#include "PropertyAnimatorCoreTimeSourceBase.h"
+#include "PropertyAnimatorCoreSystemTimeSource.generated.h"
+
+/** Enumerates all possible modes for the machine clock time source */
+UENUM(BlueprintType)
+enum class EPropertyAnimatorCoreSystemMode : uint8
+{
+	/** Local time of the machine */
+	LocalTime,
+	/** Universal time = Greenwich Mean Time */
+	UtcTime,
+	/** Specified duration elapsing until it reaches 0 */
+	Countdown,
+	/** Current time elapsed since the time source is active */
+	Stopwatch
+};
+
+/** System time source that support various option */
+UCLASS(MinimalAPI)
+class UPropertyAnimatorCoreSystemTimeSource : public UPropertyAnimatorCoreTimeSourceBase
+{
+	GENERATED_BODY()
+
+public:
+	UPropertyAnimatorCoreSystemTimeSource()
+		: UPropertyAnimatorCoreTimeSourceBase(TEXT("System"))
+	{}
+
+	//~ Begin UPropertyAnimatorTimeSourceBase
+	virtual double GetTimeElapsed() override;
+	virtual bool IsTimeSourceReady() const override;
+	virtual void OnTimeSourceActive() override;
+	//~ End UPropertyAnimatorTimeSourceBase
+
+	void SetMode(EPropertyAnimatorCoreSystemMode InMode);
+	EPropertyAnimatorCoreSystemMode GetMode() const
+	{
+		return Mode;
+	}
+
+	void SetCountdownDuration(const FTimespan& InTimeSpan);
+	void GetCountdownDuration(FTimespan& OutTimeSpan) const
+	{
+		OutTimeSpan = CountdownTimeSpan;
+	}
+
+	void SetCountdownDuration(const FString& InDuration);
+	const FString& GetCountdownDuration() const
+	{
+		return CountdownDuration;
+	}
+
+protected:
+	static FTimespan ParseTime(const FString& InFormat);
+
+	//~ Begin UObject
+#if WITH_EDITOR
+	virtual void PostEditChangeProperty(FPropertyChangedEvent& InPropertyChangedEvent) override;
+#endif
+	//~ End UObject
+
+	void OnModeChanged();
+
+	/** Machine time mode to use */
+	UPROPERTY(EditInstanceOnly, Setter, Getter, Category="Animator")
+	EPropertyAnimatorCoreSystemMode Mode = EPropertyAnimatorCoreSystemMode::LocalTime;
+
+	/**
+	* Countdown duration format : 
+	* 120 = 2 minutes
+	* 02:00 = 2 minutes
+	* 00:02:00 = 2 minutes
+	* 2m = 2 minutes
+	* 1h = 1 hour
+	* 120s = 2 minutes
+	*/
+	UPROPERTY(EditInstanceOnly, Setter, Getter, Category="Animator", meta=(EditCondition="Mode == EPropertyAnimatorCoreSystemMode::Countdown", EditConditionHides))
+	FString CountdownDuration = TEXT("1m");
+
+private:
+	UPROPERTY(Transient)
+	FTimespan CountdownTimeSpan;
+
+	UPROPERTY(Transient)
+	FDateTime ActivationTime;
+};
\ No newline at end of file
