MotionDesign : PropertyAnimator
- Manual time source improved to use custom player controls (for RC)
- Hiding irrelevant properties to simplify usability
- Moving icons from resources to content

diff --git a/Engine/Plugins/Experimental/PropertyAnimatorCore/Source/PropertyAnimatorCore/Private/Properties/PropertyAnimatorCoreContext.cpp b/Engine/Plugins/Experimental/PropertyAnimatorCore/Source/PropertyAnimatorCore/Private/Properties/PropertyAnimatorCoreContext.cpp
--- a/Engine/Plugins/Experimental/PropertyAnimatorCore/Source/PropertyAnimatorCore/Private/Properties/PropertyAnimatorCoreContext.cpp
+++ b/Engine/Plugins/Experimental/PropertyAnimatorCore/Source/PropertyAnimatorCore/Private/Properties/PropertyAnimatorCoreContext.cpp
@@ -190,6 +190,11 @@ void UPropertyAnimatorCoreContext::PostEditChangeProperty(FPropertyChangedEvent&
 }
 #endif
 
+void UPropertyAnimatorCoreContext::OnAnimatedPropertyLinked()
+{
+	bEditMagnitude = AnimatedProperty.IsA<FNumericProperty>() || AnimatedProperty.HasA<FNumericProperty>();
+}
+
 void UPropertyAnimatorCoreContext::OnModeChanged()
 {
 	if (const UPropertyAnimatorCoreHandlerBase* Handler = GetHandler())
diff --git a/Engine/Plugins/Experimental/PropertyAnimatorCore/Source/PropertyAnimatorCore/Private/Properties/PropertyAnimatorCoreData.cpp b/Engine/Plugins/Experimental/PropertyAnimatorCore/Source/PropertyAnimatorCore/Private/Properties/PropertyAnimatorCoreData.cpp
--- a/Engine/Plugins/Experimental/PropertyAnimatorCore/Source/PropertyAnimatorCore/Private/Properties/PropertyAnimatorCoreData.cpp
+++ b/Engine/Plugins/Experimental/PropertyAnimatorCore/Source/PropertyAnimatorCore/Private/Properties/PropertyAnimatorCoreData.cpp
@@ -327,6 +327,44 @@ TOptional<FPropertyAnimatorCoreData> FPropertyAnimatorCoreData::GetRootParent()
 	return FPropertyAnimatorCoreData(GetOwner(), ParentChainProperties, GetPropertyResolverClass());
 }
 
+TArray<FPropertyAnimatorCoreData> FPropertyAnimatorCoreData::GetChildrenProperties(int32 InDepthSearch) const
+{
+	TArray<FPropertyAnimatorCoreData> ChildrenProperties;
+	FProperty* LeafProperty = GetLeafProperty();
+
+	if (!LeafProperty || InDepthSearch-- <= 0)
+	{
+		return ChildrenProperties;
+	}
+
+	if (const FStructProperty* StructProperty = CastField<FStructProperty>(LeafProperty))
+	{
+		if (const UScriptStruct* Struct = StructProperty->Struct)
+		{
+			TArray<FProperty*> Properties;
+			Algo::Transform(ChainProperties, Properties, [](const TFieldPath<FProperty>& InProperty)
+			{
+				return InProperty.Get();
+			});
+
+			for (FProperty* ChildProperty : TFieldRange<FProperty>(Struct))
+			{
+				if (ChildProperty)
+				{
+					FPropertyAnimatorCoreData ChildPropertyData(GetOwner(), Properties, ChildProperty, GetPropertyResolverClass());
+
+					ChildrenProperties.Add(ChildPropertyData);
+					ChildrenProperties.Append(ChildPropertyData.GetChildrenProperties(InDepthSearch));
+				}
+			}
+		}
+	}
+
+	// We don't handle object, array, map, set container for now
+
+	return ChildrenProperties;
+}
+
 UPropertyAnimatorCoreHandlerBase* FPropertyAnimatorCoreData::GetPropertyHandler() const
 {
 	// Cache it once
diff --git a/Engine/Plugins/Experimental/PropertyAnimatorCore/Source/PropertyAnimatorCore/Private/TimeSources/PropertyAnimatorCoreManualTimeSource.cpp b/Engine/Plugins/Experimental/PropertyAnimatorCore/Source/PropertyAnimatorCore/Private/TimeSources/PropertyAnimatorCoreManualTimeSource.cpp
--- a/Engine/Plugins/Experimental/PropertyAnimatorCore/Source/PropertyAnimatorCore/Private/TimeSources/PropertyAnimatorCoreManualTimeSource.cpp
+++ b/Engine/Plugins/Experimental/PropertyAnimatorCore/Source/PropertyAnimatorCore/Private/TimeSources/PropertyAnimatorCoreManualTimeSource.cpp
@@ -2,9 +2,22 @@
 
 #include "TimeSources/PropertyAnimatorCoreManualTimeSource.h"
 
+#include "Misc/App.h"
+
+void UPropertyAnimatorCoreManualTimeSource::SetOverrideTime(bool bInOverride)
+{
+	if (bOverrideTime == bInOverride)
+	{
+		return;
+	}
+
+	bOverrideTime = bInOverride;
+	OnOverrideTimeChanged();
+}
+
 void UPropertyAnimatorCoreManualTimeSource::SetCustomTime(double InTime)
 {
-	if (FMath::IsNearlyEqual(InTime, CustomTime))
+	if (!bOverrideTime || FMath::IsNearlyEqual(InTime, CustomTime))
 	{
 		return;
 	}
@@ -12,12 +25,150 @@ void UPropertyAnimatorCoreManualTimeSource::SetCustomTime(double InTime)
 	CustomTime = InTime;
 }
 
+void UPropertyAnimatorCoreManualTimeSource::SetState(const FPropertyAnimatorCoreManualState& InState)
+{
+	if (State.Status == InState.Status)
+	{
+		return;
+	}
+
+	State = InState;
+	OnStateChanged();
+}
+
 double UPropertyAnimatorCoreManualTimeSource::GetTimeElapsed()
 {
+	if (!bOverrideTime)
+	{
+		/*
+		 * Don't use world delta time to avoid time dilation,
+		 * get the app to use raw time between frames and increment when this time source is enabled
+		 */
+		if (ActiveStatus == EPropertyAnimatorCoreManualStatus::PlayingForward)
+		{
+			CustomTime += FApp::GetDeltaTime();
+		}
+		else if (ActiveStatus == EPropertyAnimatorCoreManualStatus::PlayingBackward)
+		{
+			CustomTime -= FApp::GetDeltaTime();
+		}
+	}
+
 	return CustomTime;
 }
 
 bool UPropertyAnimatorCoreManualTimeSource::IsTimeSourceReady() const
 {
-	return true;
-}
\ No newline at end of file
+	if (bOverrideTime)
+	{
+		return true;
+	}
+
+	return ActiveStatus != EPropertyAnimatorCoreManualStatus::Stopped;
+}
+
+void UPropertyAnimatorCoreManualTimeSource::OnTimeSourceActive()
+{
+	Super::OnTimeSourceActive();
+
+	ActiveStatus = EPropertyAnimatorCoreManualStatus::Stopped;
+}
+
+void UPropertyAnimatorCoreManualTimeSource::OnTimeSourceInactive()
+{
+	Super::OnTimeSourceInactive();
+
+	Stop();
+}
+
+void UPropertyAnimatorCoreManualTimeSource::Play(bool bInForward)
+{
+	if (bOverrideTime)
+	{
+		return;
+	}
+
+	// Allow change from playing forward to backward
+	if (!IsPlaying()
+		|| (bInForward && ActiveStatus == EPropertyAnimatorCoreManualStatus::PlayingBackward)
+		|| (!bInForward && ActiveStatus == EPropertyAnimatorCoreManualStatus::PlayingForward))
+	{
+		ActiveStatus = bInForward ? EPropertyAnimatorCoreManualStatus::PlayingForward : EPropertyAnimatorCoreManualStatus::PlayingBackward;
+	}
+}
+
+void UPropertyAnimatorCoreManualTimeSource::Pause()
+{
+	if (!IsPlaying())
+	{
+		return;
+	}
+
+	ActiveStatus = EPropertyAnimatorCoreManualStatus::Paused;
+}
+
+void UPropertyAnimatorCoreManualTimeSource::Stop()
+{
+	if (ActiveStatus == EPropertyAnimatorCoreManualStatus::Stopped)
+	{
+		return;
+	}
+
+	CustomTime = 0.f;
+	Pause();
+	ActiveStatus = EPropertyAnimatorCoreManualStatus::Stopped;
+}
+
+EPropertyAnimatorCoreManualStatus UPropertyAnimatorCoreManualTimeSource::GetPlaybackStatus() const
+{
+	return bOverrideTime ? EPropertyAnimatorCoreManualStatus::Stopped : ActiveStatus;
+}
+
+bool UPropertyAnimatorCoreManualTimeSource::IsPlaying() const
+{
+	return !bOverrideTime && (ActiveStatus == EPropertyAnimatorCoreManualStatus::PlayingBackward || ActiveStatus == EPropertyAnimatorCoreManualStatus::PlayingForward);
+}
+
+#if WITH_EDITOR
+void UPropertyAnimatorCoreManualTimeSource::PostEditChangeProperty(FPropertyChangedEvent& InPropertyChangedEvent)
+{
+	Super::PostEditChangeProperty(InPropertyChangedEvent);
+
+	const FName MemberName = InPropertyChangedEvent.GetMemberPropertyName();
+
+	if (MemberName == GET_MEMBER_NAME_CHECKED(UPropertyAnimatorCoreManualTimeSource, State))
+	{
+		OnStateChanged();
+	}
+	else if (MemberName == GET_MEMBER_NAME_CHECKED(UPropertyAnimatorCoreManualTimeSource, bOverrideTime))
+	{
+		OnOverrideTimeChanged();
+	}
+}
+#endif
+
+void UPropertyAnimatorCoreManualTimeSource::OnStateChanged()
+{
+	switch (State.Status)
+	{
+	case EPropertyAnimatorCoreManualStatus::Stopped:
+		Stop();
+		break;
+	case EPropertyAnimatorCoreManualStatus::Paused:
+		Pause();
+		break;
+	case EPropertyAnimatorCoreManualStatus::PlayingForward:
+		Play(/** Forward */true);
+		break;
+	case EPropertyAnimatorCoreManualStatus::PlayingBackward:
+		Play(/** Forward */false);
+		break;
+	}
+}
+
+void UPropertyAnimatorCoreManualTimeSource::OnOverrideTimeChanged()
+{
+	Stop();
+	CustomTime = 0.f;
+	State.Status = EPropertyAnimatorCoreManualStatus::Stopped;
+}
diff --git a/Engine/Plugins/Experimental/PropertyAnimatorCore/Source/PropertyAnimatorCore/Public/Properties/PropertyAnimatorCoreContext.h b/Engine/Plugins/Experimental/PropertyAnimatorCore/Source/PropertyAnimatorCore/Public/Properties/PropertyAnimatorCoreContext.h
--- a/Engine/Plugins/Experimental/PropertyAnimatorCore/Source/PropertyAnimatorCore/Public/Properties/PropertyAnimatorCoreContext.h
+++ b/Engine/Plugins/Experimental/PropertyAnimatorCore/Source/PropertyAnimatorCore/Public/Properties/PropertyAnimatorCoreContext.h
@@ -104,7 +104,7 @@ protected:
 	//~ End UObject
 
 	/** Called once, when the property is linked to this context */
-	virtual void OnAnimatedPropertyLinked() {}
+	PROPERTYANIMATORCORE_API virtual void OnAnimatedPropertyLinked();
 
 	/** Called when the animated property owner is updated */
 	virtual void OnAnimatedPropertyOwnerUpdated(UObject* InPreviousOwner, UObject* InNewOwner) {}
@@ -147,8 +147,12 @@ private:
 	UPROPERTY(EditInstanceOnly, Setter="SetAnimated", Getter="IsAnimated", Category="Animator", meta=(AllowPrivateAccess="true"))
 	bool bAnimated = true;
 
+	/** Edit condition for Magnitude */
+	UPROPERTY(Transient)
+	bool bEditMagnitude = true;
+
 	/** Magnitude of the effect on this property */
-	UPROPERTY(EditInstanceOnly, Setter, Getter, Category="Animator", meta=(ClampMin="0.0", ClampMax="1.0"))
+	UPROPERTY(EditInstanceOnly, Setter, Getter, Category="Animator", meta=(ClampMin="0.0", ClampMax="1.0", HideEditConditionToggle, EditCondition="bEditMagnitude", EditConditionHides, AllowPrivateAccess="true"))
 	float Magnitude = 1.f;
 
 	/** Edit condition for modes */
@@ -168,7 +172,7 @@ private:
 	FInstancedStruct ConverterRule;
 
 	/** The unique group name that manages this property */
-	UPROPERTY(EditInstanceOnly, Category="Animator", Setter, Getter, meta=(GetOptions="GetSupportedGroupNames", AllowPrivateAccess="true"))
+	UPROPERTY()
 	FName GroupName = NAME_None;
 
 	/** Active group of this property */
diff --git a/Engine/Plugins/Experimental/PropertyAnimatorCore/Source/PropertyAnimatorCore/Public/Properties/PropertyAnimatorCoreData.h b/Engine/Plugins/Experimental/PropertyAnimatorCore/Source/PropertyAnimatorCore/Public/Properties/PropertyAnimatorCoreData.h
--- a/Engine/Plugins/Experimental/PropertyAnimatorCore/Source/PropertyAnimatorCore/Public/Properties/PropertyAnimatorCoreData.h
+++ b/Engine/Plugins/Experimental/PropertyAnimatorCore/Source/PropertyAnimatorCore/Public/Properties/PropertyAnimatorCoreData.h
@@ -143,6 +143,7 @@ struct FPropertyAnimatorCoreData
 	/** Returns the top most parent / member property if there is one */
 	PROPERTYANIMATORCORE_API TOptional<FPropertyAnimatorCoreData> GetRootParent() const;
 
+	/** Checks if this property is of a specific type */
 	template<typename InPropertyClass
 		UE_REQUIRES(std::is_base_of_v<FProperty, InPropertyClass>)>
 	bool IsA() const
@@ -155,6 +156,25 @@ struct FPropertyAnimatorCoreData
 		return false;
 	}
 
+	/** Gets the children of this property */
+	TArray<FPropertyAnimatorCoreData> GetChildrenProperties(int32 InDepthSearch = 3) const;
+
+	/** Checks if this property contains a specific type */
+	template<typename InPropertyClass
+		UE_REQUIRES(std::is_base_of_v<FProperty, InPropertyClass>)>
+	bool HasA() const
+	{
+		for (const FPropertyAnimatorCoreData& ChildProperty : GetChildrenProperties())
+		{
+			if (ChildProperty.IsA<InPropertyClass>())
+			{
+				return true;
+			}
+		}
+
+		return false;
+	}
+
 	/**
 	 * Gets the value ptr for this property
 	 * will use getter if available, otherwise directly access the property address
diff --git a/Engine/Plugins/Experimental/PropertyAnimatorCore/Source/PropertyAnimatorCore/Public/TimeSources/PropertyAnimatorCoreManualTimeSource.h b/Engine/Plugins/Experimental/PropertyAnimatorCore/Source/PropertyAnimatorCore/Public/TimeSources/PropertyAnimatorCoreManualTimeSource.h
--- a/Engine/Plugins/Experimental/PropertyAnimatorCore/Source/PropertyAnimatorCore/Public/TimeSources/PropertyAnimatorCoreManualTimeSource.h
+++ b/Engine/Plugins/Experimental/PropertyAnimatorCore/Source/PropertyAnimatorCore/Public/TimeSources/PropertyAnimatorCoreManualTimeSource.h
@@ -5,6 +5,29 @@
 #include "PropertyAnimatorCoreTimeSourceBase.h"
 #include "PropertyAnimatorCoreManualTimeSource.generated.h"
 
+UENUM()
+enum class EPropertyAnimatorCoreManualStatus : uint8
+{
+	/** Animation is done */
+	Stopped = 1 << 0,
+	/** Animation is paused */
+	Paused = 1 << 1,
+	/** Animation is playing */
+	PlayingForward = 1 << 2,
+	/** Animation is playing in reverse */
+	PlayingBackward = 1 << 3
+};
+
+/** Used to store status of player, for type customization */
+USTRUCT()
+struct FPropertyAnimatorCoreManualState
+{
+	GENERATED_BODY()
+
+	UPROPERTY(EditInstanceOnly, Category="Animator")
+	EPropertyAnimatorCoreManualStatus Status = EPropertyAnimatorCoreManualStatus::Stopped;
+};
+
 UCLASS(MinimalAPI)
 class UPropertyAnimatorCoreManualTimeSource : public UPropertyAnimatorCoreTimeSourceBase
 {
@@ -15,19 +38,60 @@ public:
 		: UPropertyAnimatorCoreTimeSourceBase(TEXT("Manual"))
 	{}
 
+	void SetOverrideTime(bool bInOverride);
+	bool GetOverrideTime() const
+	{
+		return bOverrideTime;
+	}
+
 	PROPERTYANIMATORCORE_API void SetCustomTime(double InTime);
 	double GetCustomTime() const
 	{
 		return CustomTime;
 	}
 
+	void SetState(const FPropertyAnimatorCoreManualState& InState);
+	const FPropertyAnimatorCoreManualState& GetState() const
+	{
+		return State;
+	}
+
 	//~ Begin UPropertyAnimatorTimeSourceBase
 	virtual double GetTimeElapsed() override;
 	virtual bool IsTimeSourceReady() const override;
+	virtual void OnTimeSourceActive() override;
+	virtual void OnTimeSourceInactive() override;
 	//~ End UPropertyAnimatorTimeSourceBase
 
+	void Play(bool bInForward);
+	void Pause();
+	void Stop();
+
+	EPropertyAnimatorCoreManualStatus GetPlaybackStatus() const;
+	bool IsPlaying() const;
+
 protected:
+	//~ Begin UObject
+#if WITH_EDITOR
+	virtual void PostEditChangeProperty(FPropertyChangedEvent& InPropertyChangedEvent) override;
+#endif
+	//~ End UObject
+
+	void OnStateChanged();
+	void OnOverrideTimeChanged();
+
 	/** Allows you to drive controllers with this float */
-	UPROPERTY(EditInstanceOnly, BlueprintReadWrite, Setter="SetCustomTime", Getter="GetCustomTime", Category="Animator")
+	UPROPERTY(EditInstanceOnly, Setter="SetOverrideTime", Getter="GetOverrideTime", Category="Animator", meta=(InlineEditConditionToggle))
+	bool bOverrideTime = true;
+
+	/** Time to evaluate */
+	UPROPERTY(EditInstanceOnly, Setter, Getter, Category="Animator", meta=(EditCondition="bOverrideTime"))
 	double CustomTime = 0.f;
+
+	/** Playback state */
+	UPROPERTY(EditInstanceOnly, Transient, DuplicateTransient, Setter, Getter, Category="Animator", meta=(EditCondition="!bOverrideTime", EditConditionHides, HideEditConditionToggles))
+	FPropertyAnimatorCoreManualState State;
+
+	/** Current active status for the player */
+	EPropertyAnimatorCoreManualStatus ActiveStatus = EPropertyAnimatorCoreManualStatus::Stopped;
 };
\ No newline at end of file
diff --git a/Engine/Plugins/Experimental/PropertyAnimatorCore/Source/PropertyAnimatorCoreEditor/Private/Customizations/PropertyAnimatorCoreEditorManualStateTypeCustomization.cpp b/Engine/Plugins/Experimental/PropertyAnimatorCore/Source/PropertyAnimatorCoreEditor/Private/Customizations/PropertyAnimatorCoreEditorManualStateTypeCustomization.cpp
new file mode 100644
--- /dev/null
+++ b/Engine/Plugins/Experimental/PropertyAnimatorCore/Source/PropertyAnimatorCoreEditor/Private/Customizations/PropertyAnimatorCoreEditorManualStateTypeCustomization.cpp
@@ -0,0 +1,168 @@
+// Copyright Epic Games, Inc. All Rights Reserved.
+
+#include "Customizations/PropertyAnimatorCoreEditorManualStateTypeCustomization.h"
+
+#include "DetailWidgetRow.h"
+#include "Styles/PropertyAnimatorCoreEditorStyle.h"
+#include "Styling/StyleColors.h"
+#include "TimeSources/PropertyAnimatorCoreManualTimeSource.h"
+#include "Widgets/Images/SImage.h"
+#include "Widgets/Input/SButton.h"
+#include "Widgets/SBoxPanel.h"
+
+void FPropertyAnimatorCoreEditorManualStateTypeCustomization::CustomizeHeader(TSharedRef<IPropertyHandle> InPropertyHandle, FDetailWidgetRow& InRow, IPropertyTypeCustomizationUtils& InUtils)
+{
+	if (!InPropertyHandle->IsValidHandle())
+	{
+		return;
+	}
+
+	StatusPropertyHandle = InPropertyHandle->GetChildHandle(GET_MEMBER_NAME_CHECKED(FPropertyAnimatorCoreManualState, Status));
+
+	InRow.NameContent()
+	[
+		InPropertyHandle->CreatePropertyNameWidget()
+	];
+
+	const FVector2D ImageSize(16.f);
+
+	InRow.ValueContent()
+	.HAlign(HAlign_Fill)
+	[
+		SNew(SHorizontalBox)
+		+ SHorizontalBox::Slot()
+		.AutoWidth()
+		.Padding(0.f, 0.f, 0.f, 5.f)
+		.HAlign(HAlign_Fill)
+		[
+			SNew(SButton)
+			.HAlign(HAlign_Fill)
+			.ButtonColorAndOpacity(this, &FPropertyAnimatorCoreEditorManualStateTypeCustomization::IsPlaybackStatusActive, EPropertyAnimatorCoreManualStatus::PlayingBackward)
+			.OnClicked(this, &FPropertyAnimatorCoreEditorManualStateTypeCustomization::SetPlaybackStatus, EPropertyAnimatorCoreManualStatus::PlayingBackward)
+			.IsEnabled(this, &FPropertyAnimatorCoreEditorManualStateTypeCustomization::IsPlaybackStatusAllowed, EPropertyAnimatorCoreManualStatus::PlayingBackward)
+			[
+				SNew(SImage)
+				.DesiredSizeOverride(ImageSize)
+				.Image(FPropertyAnimatorCoreEditorStyle::Get().GetBrush(TEXT("ManualTimeSourceControl.PlayBackward")))
+			]
+		]
+		+ SHorizontalBox::Slot()
+		.AutoWidth()
+		.Padding(0.f, 0.f, 0.f, 5.f)
+		.HAlign(HAlign_Fill)
+		[
+			SNew(SButton)
+			.HAlign(HAlign_Fill)
+			.ButtonColorAndOpacity(this, &FPropertyAnimatorCoreEditorManualStateTypeCustomization::IsPlaybackStatusActive, EPropertyAnimatorCoreManualStatus::PlayingForward)
+			.OnClicked(this, &FPropertyAnimatorCoreEditorManualStateTypeCustomization::SetPlaybackStatus, EPropertyAnimatorCoreManualStatus::PlayingForward)
+			.IsEnabled(this, &FPropertyAnimatorCoreEditorManualStateTypeCustomization::IsPlaybackStatusAllowed, EPropertyAnimatorCoreManualStatus::PlayingForward)
+			[
+				SNew(SImage)
+				.DesiredSizeOverride(ImageSize)
+				.Image(FPropertyAnimatorCoreEditorStyle::Get().GetBrush(TEXT("ManualTimeSourceControl.PlayForward")))
+			]
+		]
+		+ SHorizontalBox::Slot()
+		.AutoWidth()
+		.Padding(0.f, 0.f, 0.f, 5.f)
+		.HAlign(HAlign_Fill)
+		[
+			SNew(SButton)
+			.HAlign(HAlign_Fill)
+			.ButtonColorAndOpacity(this, &FPropertyAnimatorCoreEditorManualStateTypeCustomization::IsPlaybackStatusActive, EPropertyAnimatorCoreManualStatus::Paused)
+			.OnClicked(this, &FPropertyAnimatorCoreEditorManualStateTypeCustomization::SetPlaybackStatus, EPropertyAnimatorCoreManualStatus::Paused)
+			.IsEnabled(this, &FPropertyAnimatorCoreEditorManualStateTypeCustomization::IsPlaybackStatusAllowed, EPropertyAnimatorCoreManualStatus::Paused)
+			[
+				SNew(SImage)
+				.DesiredSizeOverride(ImageSize)
+				.Image(FPropertyAnimatorCoreEditorStyle::Get().GetBrush(TEXT("ManualTimeSourceControl.Pause")))
+			]
+		]
+		+ SHorizontalBox::Slot()
+		.AutoWidth()
+		.Padding(0.f, 0.f, 0.f, 5.f)
+		.HAlign(HAlign_Fill)
+		[
+			SNew(SButton)
+			.HAlign(HAlign_Fill)
+			.ButtonColorAndOpacity(this, &FPropertyAnimatorCoreEditorManualStateTypeCustomization::IsPlaybackStatusActive, EPropertyAnimatorCoreManualStatus::Stopped)
+			.OnClicked(this, &FPropertyAnimatorCoreEditorManualStateTypeCustomization::SetPlaybackStatus, EPropertyAnimatorCoreManualStatus::Stopped)
+			.IsEnabled(this, &FPropertyAnimatorCoreEditorManualStateTypeCustomization::IsPlaybackStatusAllowed, EPropertyAnimatorCoreManualStatus::Stopped)
+			[
+				SNew(SImage)
+				.DesiredSizeOverride(ImageSize)
+				.Image(FPropertyAnimatorCoreEditorStyle::Get().GetBrush(TEXT("ManualTimeSourceControl.Stop")))
+			]
+		]
+	];
+}
+
+void FPropertyAnimatorCoreEditorManualStateTypeCustomization::CustomizeChildren(TSharedRef<IPropertyHandle> InPropertyHandle, IDetailChildrenBuilder& InBuilder, IPropertyTypeCustomizationUtils& InUtils)
+{
+}
+
+FReply FPropertyAnimatorCoreEditorManualStateTypeCustomization::SetPlaybackStatus(EPropertyAnimatorCoreManualStatus InStatus)
+{
+	if (StatusPropertyHandle.IsValid() && StatusPropertyHandle->IsValidHandle())
+	{
+		StatusPropertyHandle->SetValue(static_cast<uint8>(InStatus));
+	}
+
+	return FReply::Handled();
+}
+
+bool FPropertyAnimatorCoreEditorManualStateTypeCustomization::IsPlaybackStatusAllowed(EPropertyAnimatorCoreManualStatus InStatus) const
+{
+	if (StatusPropertyHandle.IsValid() && StatusPropertyHandle->IsValidHandle())
+	{
+		uint8 CurrentStatus = 0;
+		if (StatusPropertyHandle->GetValue(CurrentStatus) == FPropertyAccess::Success)
+		{
+			switch(static_cast<EPropertyAnimatorCoreManualStatus>(CurrentStatus))
+			{
+			case EPropertyAnimatorCoreManualStatus::Stopped:
+				{
+					return InStatus == EPropertyAnimatorCoreManualStatus::PlayingForward
+						|| InStatus == EPropertyAnimatorCoreManualStatus::PlayingBackward;
+				}
+			case EPropertyAnimatorCoreManualStatus::Paused:
+				{
+					return InStatus == EPropertyAnimatorCoreManualStatus::PlayingForward
+						|| InStatus == EPropertyAnimatorCoreManualStatus::PlayingBackward
+						|| InStatus == EPropertyAnimatorCoreManualStatus::Stopped;
+				}
+			case EPropertyAnimatorCoreManualStatus::PlayingForward:
+				{
+					return InStatus == EPropertyAnimatorCoreManualStatus::PlayingBackward
+						|| InStatus == EPropertyAnimatorCoreManualStatus::Paused
+						|| InStatus == EPropertyAnimatorCoreManualStatus::Stopped;
+				}
+			case EPropertyAnimatorCoreManualStatus::PlayingBackward:
+				{
+					return InStatus == EPropertyAnimatorCoreManualStatus::PlayingForward
+						|| InStatus == EPropertyAnimatorCoreManualStatus::Paused
+						|| InStatus == EPropertyAnimatorCoreManualStatus::Stopped;
+				}
+			}
+		}
+	}
+
+	return false;
+}
+
+FSlateColor FPropertyAnimatorCoreEditorManualStateTypeCustomization::IsPlaybackStatusActive(EPropertyAnimatorCoreManualStatus InStatus) const
+{
+	if (StatusPropertyHandle.IsValid() && StatusPropertyHandle->IsValidHandle())
+	{
+		uint8 CurrentStatus = 0;
+		if (StatusPropertyHandle->GetValue(CurrentStatus) == FPropertyAccess::Success)
+		{
+			if (InStatus == static_cast<EPropertyAnimatorCoreManualStatus>(CurrentStatus))
+			{
+				return FStyleColors::Select;
+			}
+		}
+	}
+
+	return FStyleColors::White;
+}
diff --git a/Engine/Plugins/Experimental/PropertyAnimatorCore/Source/PropertyAnimatorCoreEditor/Private/Customizations/PropertyAnimatorCoreEditorManualStateTypeCustomization.h b/Engine/Plugins/Experimental/PropertyAnimatorCore/Source/PropertyAnimatorCoreEditor/Private/Customizations/PropertyAnimatorCoreEditorManualStateTypeCustomization.h
new file mode 100644
--- /dev/null
+++ b/Engine/Plugins/Experimental/PropertyAnimatorCore/Source/PropertyAnimatorCoreEditor/Private/Customizations/PropertyAnimatorCoreEditorManualStateTypeCustomization.h
@@ -0,0 +1,33 @@
+// Copyright Epic Games, Inc. All Rights Reserved.
+
+#pragma once
+
+#include "IPropertyTypeCustomization.h"
+#include "Templates/SharedPointer.h"
+
+class FReply;
+class IPropertyHandle;
+enum class EPropertyAnimatorCoreManualStatus : uint8;
+struct FSlateColor;
+
+/** Type customization for FPropertyAnimatorCoreManualState to show a player */
+class FPropertyAnimatorCoreEditorManualStateTypeCustomization : public IPropertyTypeCustomization
+{
+public:
+	static TSharedRef<IPropertyTypeCustomization> MakeInstance()
+	{
+		return MakeShared<FPropertyAnimatorCoreEditorManualStateTypeCustomization>();
+	}
+
+	//~ Begin IPropertyTypeCustomization
+	virtual void CustomizeHeader(TSharedRef<IPropertyHandle> InPropertyHandle, FDetailWidgetRow& InRow, IPropertyTypeCustomizationUtils& InUtils) override;
+	virtual void CustomizeChildren(TSharedRef<IPropertyHandle> InPropertyHandle, IDetailChildrenBuilder& InBuilder, IPropertyTypeCustomizationUtils& InUtils) override;
+	//~ End IPropertyTypeCustomization
+
+protected:
+	FReply SetPlaybackStatus(EPropertyAnimatorCoreManualStatus InStatus);
+	bool IsPlaybackStatusAllowed(EPropertyAnimatorCoreManualStatus InStatus) const;
+	FSlateColor IsPlaybackStatusActive(EPropertyAnimatorCoreManualStatus InStatus) const;
+
+	TSharedPtr<IPropertyHandle> StatusPropertyHandle;
+};
diff --git a/Engine/Plugins/Experimental/PropertyAnimatorCore/Source/PropertyAnimatorCoreEditor/Private/PropertyAnimatorCoreEditorModule.cpp b/Engine/Plugins/Experimental/PropertyAnimatorCore/Source/PropertyAnimatorCoreEditor/Private/PropertyAnimatorCoreEditorModule.cpp
--- a/Engine/Plugins/Experimental/PropertyAnimatorCore/Source/PropertyAnimatorCoreEditor/Private/PropertyAnimatorCoreEditorModule.cpp
+++ b/Engine/Plugins/Experimental/PropertyAnimatorCore/Source/PropertyAnimatorCoreEditor/Private/PropertyAnimatorCoreEditorModule.cpp
@@ -5,12 +5,14 @@
 #include "Animators/PropertyAnimatorCoreBase.h"
 #include "Customizations/PropertyAnimatorCoreEditorContextTypeCustomization.h"
 #include "Customizations/PropertyAnimatorCoreEditorDetailCustomization.h"
+#include "Customizations/PropertyAnimatorCoreEditorManualStateTypeCustomization.h"
 #include "Customizations/PropertyAnimatorCoreEditorSequencerTimeSourceChannelTypeCustomization.h"
 #include "ISequencerModule.h"
 #include "Modules/ModuleManager.h"
 #include "Properties/PropertyAnimatorCoreContext.h"
 #include "PropertyEditorModule.h"
 #include "Sequencer/MovieSceneAnimatorTrackEditor.h"
+#include "TimeSources/PropertyAnimatorCoreManualTimeSource.h"
 #include "TimeSources/PropertyAnimatorCoreSequencerTimeSource.h"
 
 void FPropertyAnimatorCoreEditorModule::StartupModule()
@@ -18,6 +20,7 @@ void FPropertyAnimatorCoreEditorModule::StartupModule()
 	FPropertyEditorModule& PropertyModule = FModuleManager::LoadModuleChecked<FPropertyEditorModule>("PropertyEditor");
 	PropertyModule.RegisterCustomPropertyTypeLayout(UPropertyAnimatorCoreContext::StaticClass()->GetFName(), FOnGetPropertyTypeCustomizationInstance::CreateStatic(&FPropertyAnimatorCoreEditorContextTypeCustomization::MakeInstance));
 	PropertyModule.RegisterCustomPropertyTypeLayout(FPropertyAnimatorCoreSequencerTimeSourceChannel::StaticStruct()->GetFName(), FOnGetPropertyTypeCustomizationInstance::CreateStatic(&FPropertyAnimatorCoreEditorSequencerTimeSourceChannelTypeCustomization::MakeInstance));
+	PropertyModule.RegisterCustomPropertyTypeLayout(FPropertyAnimatorCoreManualState::StaticStruct()->GetFName(), FOnGetPropertyTypeCustomizationInstance::CreateStatic(&FPropertyAnimatorCoreEditorManualStateTypeCustomization::MakeInstance));
 	PropertyModule.RegisterCustomClassLayout(UPropertyAnimatorCoreBase::StaticClass()->GetFName(), FOnGetDetailCustomizationInstance::CreateStatic(&FPropertyAnimatorCoreEditorDetailCustomization::MakeInstance));
 
 	ISequencerModule& SequencerModule = FModuleManager::LoadModuleChecked<ISequencerModule>("Sequencer");
@@ -31,6 +34,7 @@ void FPropertyAnimatorCoreEditorModule::ShutdownModule()
 		FPropertyEditorModule& PropertyModule = FModuleManager::GetModuleChecked<FPropertyEditorModule>("PropertyEditor");
 		PropertyModule.UnregisterCustomPropertyTypeLayout(UPropertyAnimatorCoreContext::StaticClass()->GetFName());
 		PropertyModule.UnregisterCustomPropertyTypeLayout(FPropertyAnimatorCoreSequencerTimeSourceChannel::StaticStruct()->GetFName());
+		PropertyModule.UnregisterCustomPropertyTypeLayout(FPropertyAnimatorCoreManualState::StaticStruct()->GetFName());
 		PropertyModule.UnregisterCustomClassLayout(UPropertyAnimatorCoreBase::StaticClass()->GetFName());
 	}
 
diff --git a/Engine/Plugins/Experimental/PropertyAnimatorCore/Source/PropertyAnimatorCoreEditor/Private/Styles/PropertyAnimatorCoreEditorStyle.cpp b/Engine/Plugins/Experimental/PropertyAnimatorCore/Source/PropertyAnimatorCoreEditor/Private/Styles/PropertyAnimatorCoreEditorStyle.cpp
--- a/Engine/Plugins/Experimental/PropertyAnimatorCore/Source/PropertyAnimatorCoreEditor/Private/Styles/PropertyAnimatorCoreEditorStyle.cpp
+++ b/Engine/Plugins/Experimental/PropertyAnimatorCore/Source/PropertyAnimatorCoreEditor/Private/Styles/PropertyAnimatorCoreEditorStyle.cpp
@@ -18,13 +18,17 @@ FPropertyAnimatorCoreEditorStyle::FPropertyAnimatorCoreEditorStyle()
 
 	check(Plugin.IsValid());
 
-	ContentRootDir = FPaths::Combine(Plugin->GetBaseDir(), TEXT("Resources"));
+	ContentRootDir = FPaths::Combine(Plugin->GetBaseDir(), TEXT("Content"));
 
 	Set("ClassIcon.MovieSceneAnimatorTrack",        new IMAGE_BRUSH_SVG("EditorIcons/PropertyAnimatorCoreDefault", Icon16x16));
 	Set("ClassIcon.PropertyAnimatorCoreComponent",  new IMAGE_BRUSH_SVG("EditorIcons/PropertyAnimatorCoreDefault", Icon16x16));
 	Set("ClassIcon.PropertyAnimatorCoreBase",       new IMAGE_BRUSH_SVG("EditorIcons/PropertyAnimatorCoreDefault", Icon16x16));
 	Set("PropertyControlIcon.Default",              new IMAGE_BRUSH_SVG("EditorIcons/PropertyAnimatorCoreDefault", Icon16x16));
 	Set("PropertyControlIcon.Linked",               new IMAGE_BRUSH_SVG("EditorIcons/PropertyAnimatorCoreLinked",  Icon16x16));
+	Set("ManualTimeSourceControl.PlayForward",      new IMAGE_BRUSH_SVG("EditorIcons/Play",                        Icon16x16));
+	Set("ManualTimeSourceControl.PlayBackward",     new IMAGE_BRUSH_SVG("EditorIcons/PlayReverse",                 Icon16x16));
+	Set("ManualTimeSourceControl.Pause",            new IMAGE_BRUSH_SVG("EditorIcons/Pause",                       Icon16x16));
+	Set("ManualTimeSourceControl.Stop",             new IMAGE_BRUSH_SVG("EditorIcons/Stop",                        Icon16x16));
 
 	FSlateStyleRegistry::RegisterSlateStyle(*this);
 }
