MotionDesign : PropertyAnimator
- Restore state when time source data is invalid or out of range (done)
- Allow time source to ease effect
- Animator track can offset time to take into account playback range
- Animator track can fade in and out to reduce magnitude on a specific time range (with section handles)

diff --git a/Engine/Plugins/Experimental/PropertyAnimatorCore/Source/PropertyAnimatorCore/Private/Components/PropertyAnimatorCoreComponent.cpp b/Engine/Plugins/Experimental/PropertyAnimatorCore/Source/PropertyAnimatorCore/Private/Components/PropertyAnimatorCoreComponent.cpp
--- a/Engine/Plugins/Experimental/PropertyAnimatorCore/Source/PropertyAnimatorCore/Private/Components/PropertyAnimatorCoreComponent.cpp
+++ b/Engine/Plugins/Experimental/PropertyAnimatorCore/Source/PropertyAnimatorCore/Private/Components/PropertyAnimatorCoreComponent.cpp
@@ -482,13 +482,16 @@ bool UPropertyAnimatorCoreComponent::EvaluateAnimators()
 	const UWorld* World = GetWorld();
 	const bool bIsSupportedWorld = IsValid(World) && (World->IsGameWorld() || World->IsEditorWorld());
 
-	if (!bIsSupportedWorld)
+	if (!bIsSupportedWorld || !ActiveAnimatorsTimeSource)
 	{
 		return false;
 	}
 
 	FInstancedPropertyBag Parameters;
 
+	FPropertyAnimatorCoreTimeSourceEvaluationData GlobalEvaluationData;
+	EPropertyAnimatorCoreTimeSourceResult GlobalTimeResult = ActiveAnimatorsTimeSource->FetchEvaluationData(GlobalEvaluationData);
+
 	for (const TObjectPtr<UPropertyAnimatorCoreBase>& Animator : PropertyAnimators)
 	{
 		if (!IsValid(Animator) || !Animator->GetAnimatorEnabled())
@@ -503,10 +506,21 @@ bool UPropertyAnimatorCoreComponent::EvaluateAnimators()
 			continue;
 		}
 
-		const TOptional<double> TimeElapsed = AnimatorTimeSource->GetConditionalTimeElapsed();
+		FPropertyAnimatorCoreTimeSourceEvaluationData AnimatorEvaluationData = GlobalEvaluationData;
+		EPropertyAnimatorCoreTimeSourceResult AnimatorTimeResult = GlobalTimeResult;
+
+		if (Animator->GetOverrideTimeSource())
+		{
+			AnimatorTimeResult = AnimatorTimeSource->FetchEvaluationData(AnimatorEvaluationData);
+		}
 
-		if (!TimeElapsed.IsSet())
+		if (AnimatorTimeResult != EPropertyAnimatorCoreTimeSourceResult::Evaluate)
 		{
+			if (AnimatorTimeResult == EPropertyAnimatorCoreTimeSourceResult::Reset)
+			{
+				Animator->RestoreProperties(/** Force */true);
+			}
+
 			continue;
 		}
 
@@ -514,10 +528,10 @@ bool UPropertyAnimatorCoreComponent::EvaluateAnimators()
 		Parameters.Reset();
 
 		Parameters.AddProperty(UPropertyAnimatorCoreBase::MagnitudeParameterName, EPropertyBagPropertyType::Float);
-		Parameters.SetValueFloat(UPropertyAnimatorCoreBase::MagnitudeParameterName, AnimatorsMagnitude);
+		Parameters.SetValueFloat(UPropertyAnimatorCoreBase::MagnitudeParameterName, AnimatorsMagnitude * AnimatorEvaluationData.Magnitude);
 
 		Parameters.AddProperty(UPropertyAnimatorCoreBase::TimeElapsedParameterName, EPropertyBagPropertyType::Double);
-		Parameters.SetValueDouble(UPropertyAnimatorCoreBase::TimeElapsedParameterName, TimeElapsed.GetValue());
+		Parameters.SetValueDouble(UPropertyAnimatorCoreBase::TimeElapsedParameterName, AnimatorEvaluationData.TimeElapsed);
 
 		Animator->EvaluateAnimator(Parameters);
 	}
diff --git a/Engine/Plugins/Experimental/PropertyAnimatorCore/Source/PropertyAnimatorCore/Private/Sequencer/MovieSceneAnimatorEvalTemplate.h b/Engine/Plugins/Experimental/PropertyAnimatorCore/Source/PropertyAnimatorCore/Private/Sequencer/MovieSceneAnimatorEvalTemplate.h
--- a/Engine/Plugins/Experimental/PropertyAnimatorCore/Source/PropertyAnimatorCore/Private/Sequencer/MovieSceneAnimatorEvalTemplate.h
+++ b/Engine/Plugins/Experimental/PropertyAnimatorCore/Source/PropertyAnimatorCore/Private/Sequencer/MovieSceneAnimatorEvalTemplate.h
@@ -2,8 +2,9 @@
 
 #pragma once
 
-#include "MovieSceneAnimatorExecutionToken.h"
 #include "Evaluation/MovieSceneEvalTemplate.h"
+#include "MovieSceneAnimatorExecutionToken.h"
+#include "Sequencer/MovieSceneAnimatorTypes.h"
 #include "MovieSceneAnimatorEvalTemplate.generated.h"
 
 USTRUCT()
@@ -13,8 +14,8 @@ struct FMovieSceneAnimatorEvalTemplate : public FMovieSceneEvalTemplate
 
 	FMovieSceneAnimatorEvalTemplate() = default;
 
-	explicit FMovieSceneAnimatorEvalTemplate(uint8 InChannel)
-		: Channel(InChannel)
+	explicit FMovieSceneAnimatorEvalTemplate(const FMovieSceneAnimatorSectionData& InSectionData)
+		: SectionData(InSectionData)
 	{}
 
 	//~ Begin FMovieSceneEvalTemplate
@@ -25,11 +26,11 @@ struct FMovieSceneAnimatorEvalTemplate : public FMovieSceneEvalTemplate
 
 	virtual void Evaluate(const FMovieSceneEvaluationOperand& InOperand, const FMovieSceneContext& InContext, const FPersistentEvaluationData& InPersistentData, FMovieSceneExecutionTokens& InExecutionTokens) const override
 	{
-		FMovieSceneAnimatorExecutionToken ExecutionToken(Channel);
+		FMovieSceneAnimatorExecutionToken ExecutionToken(SectionData);
 		InExecutionTokens.Add(MoveTemp(ExecutionToken));
 	}
 	//~ End FMovieSceneEvalTemplate
 
 private:
-	uint8 Channel = 0;
+	FMovieSceneAnimatorSectionData SectionData;
 };
diff --git a/Engine/Plugins/Experimental/PropertyAnimatorCore/Source/PropertyAnimatorCore/Private/Sequencer/MovieSceneAnimatorExecutionToken.h b/Engine/Plugins/Experimental/PropertyAnimatorCore/Source/PropertyAnimatorCore/Private/Sequencer/MovieSceneAnimatorExecutionToken.h
--- a/Engine/Plugins/Experimental/PropertyAnimatorCore/Source/PropertyAnimatorCore/Private/Sequencer/MovieSceneAnimatorExecutionToken.h
+++ b/Engine/Plugins/Experimental/PropertyAnimatorCore/Source/PropertyAnimatorCore/Private/Sequencer/MovieSceneAnimatorExecutionToken.h
@@ -3,22 +3,77 @@
 #pragma once
 
 #include "Evaluation/MovieSceneExecutionTokens.h"
+#include "Sequencer/MovieSceneAnimatorSection.h"
+#include "Sequencer/MovieSceneAnimatorTypes.h"
 #include "TimeSources/PropertyAnimatorCoreSequencerTimeSource.h"
 
-struct FMovieSceneAnimatorExecutionToken : IMovieSceneExecutionToken
+/** Used to restore state back to previous when outside section */
+struct FMovieSceneAnimatorPreAnimatedTokenProducer : IMovieScenePreAnimatedGlobalTokenProducer
 {
-	FMovieSceneAnimatorExecutionToken(uint8 InChannel)
+	explicit FMovieSceneAnimatorPreAnimatedTokenProducer(uint8 InChannel)
 		: Channel(InChannel)
 	{}
 
+	//~ Begin IMovieScenePreAnimatedGlobalTokenProducer
+	virtual IMovieScenePreAnimatedGlobalTokenPtr CacheExistingState() const override
+	{
+		struct FMovieSceneAnimatorPreAnimatedGlobalToken : IMovieScenePreAnimatedGlobalToken
+		{
+			FMovieSceneAnimatorPreAnimatedGlobalToken(uint8 InChannel)
+				: Channel(InChannel)
+			{}
+
+			//~ Begin IMovieScenePreAnimatedGlobalToken
+			virtual void RestoreState(const UE::MovieScene::FRestoreStateParams& InParams) override
+			{
+				UPropertyAnimatorCoreSequencerTimeSource::OnAnimatorTimeEvaluated.Broadcast(Channel, TOptional<double>(), TOptional<float>());
+			}
+			//~ End IMovieScenePreAnimatedGlobalToken
+
+		private:
+			uint8 Channel = 0;
+		};
+
+		return FMovieSceneAnimatorPreAnimatedGlobalToken(Channel);
+	}
+	//~ End IMovieScenePreAnimatedGlobalTokenProducer
+
+	static FMovieSceneAnimTypeID GetAnimTypeID()
+	{
+		return TMovieSceneAnimTypeID<FMovieSceneAnimatorPreAnimatedTokenProducer>();
+	}
+
+private:
+	uint8 Channel = 0;
+};
+
+/** Used to evaluate active section */
+struct FMovieSceneAnimatorExecutionToken : IMovieSceneExecutionToken
+{
+	FMovieSceneAnimatorExecutionToken(const FMovieSceneAnimatorSectionData& InSectionData)
+		: SectionData(InSectionData)
+	{}
+
 	//~ Begin IMovieSceneExecutionToken
 	virtual void Execute(const FMovieSceneContext& InContext, const FMovieSceneEvaluationOperand& InOperand, FPersistentEvaluationData& InPersistentData, IMovieScenePlayer& InPlayer) override
 	{
-		const double EvaluatedTime = InContext.GetTime().AsDecimal() / InContext.GetFrameRate().AsDecimal();
-		UPropertyAnimatorCoreSequencerTimeSource::OnAnimatorTimeEvaluated.Broadcast(Channel, EvaluatedTime);
+		InPlayer.SavePreAnimatedState(FMovieSceneAnimatorPreAnimatedTokenProducer::GetAnimTypeID(), FMovieSceneAnimatorPreAnimatedTokenProducer(SectionData.Channel));
+
+		double EvaluatedTime = InContext.GetTime().AsDecimal();
+
+		if (SectionData.bUseSectionTime && SectionData.Section)
+		{
+			EvaluatedTime -= SectionData.Section->GetInclusiveStartFrame().Value;
+		}
+
+		EvaluatedTime /= InContext.GetFrameRate().AsDecimal();
+
+		const float Magnitude = SectionData.Section ? SectionData.Section->EvaluateEasing(InContext.GetTime()) : 1;
+
+		UPropertyAnimatorCoreSequencerTimeSource::OnAnimatorTimeEvaluated.Broadcast(SectionData.Channel, EvaluatedTime, Magnitude);
 	}
 	//~ End IMovieSceneExecutionToken
 
 private:
-	uint8 Channel = 0;
+	FMovieSceneAnimatorSectionData SectionData;
 };
diff --git a/Engine/Plugins/Experimental/PropertyAnimatorCore/Source/PropertyAnimatorCore/Private/Sequencer/MovieSceneAnimatorSection.cpp b/Engine/Plugins/Experimental/PropertyAnimatorCore/Source/PropertyAnimatorCore/Private/Sequencer/MovieSceneAnimatorSection.cpp
--- a/Engine/Plugins/Experimental/PropertyAnimatorCore/Source/PropertyAnimatorCore/Private/Sequencer/MovieSceneAnimatorSection.cpp
+++ b/Engine/Plugins/Experimental/PropertyAnimatorCore/Source/PropertyAnimatorCore/Private/Sequencer/MovieSceneAnimatorSection.cpp
@@ -6,9 +6,15 @@ UMovieSceneAnimatorSection::UMovieSceneAnimatorSection()
 	: UMovieSceneSection()
 {
 	bSupportsInfiniteRange = true;
+	EvalOptions.CompletionMode = EMovieSceneCompletionMode::RestoreState;
 }
 
 void UMovieSceneAnimatorSection::SetChannel(uint8 InChannel)
 {
 	Channel = InChannel;
 }
+
+void UMovieSceneAnimatorSection::SetUseSectionTime(bool bInUse)
+{
+	bUseSectionTime = bInUse;
+}
diff --git a/Engine/Plugins/Experimental/PropertyAnimatorCore/Source/PropertyAnimatorCore/Private/Sequencer/MovieSceneAnimatorTrack.cpp b/Engine/Plugins/Experimental/PropertyAnimatorCore/Source/PropertyAnimatorCore/Private/Sequencer/MovieSceneAnimatorTrack.cpp
--- a/Engine/Plugins/Experimental/PropertyAnimatorCore/Source/PropertyAnimatorCore/Private/Sequencer/MovieSceneAnimatorTrack.cpp
+++ b/Engine/Plugins/Experimental/PropertyAnimatorCore/Source/PropertyAnimatorCore/Private/Sequencer/MovieSceneAnimatorTrack.cpp
@@ -13,7 +13,8 @@
 UMovieSceneAnimatorTrack::UMovieSceneAnimatorTrack()
 	: UMovieSceneNameableTrack()
 {
-	SupportedBlendTypes = FMovieSceneBlendTypeField::None();
+	// Needed for easing
+	SupportedBlendTypes.Add(EMovieSceneBlendType::Absolute);
 }
 
 int32 UMovieSceneAnimatorTrack::GetChannelCount(uint8 InChannel) const
@@ -47,6 +48,10 @@ UMovieSceneSection* UMovieSceneAnimatorTrack::CreateNewSection()
 	{
 		NewSection->SetStartFrame(MovieScene->GetPlaybackRange().GetLowerBound());
 		NewSection->SetEndFrame(MovieScene->GetPlaybackRange().GetUpperBound());
+
+		// For easing
+		NewSection->SetBlendType(EMovieSceneBlendType::Absolute);
+		UpdateEasing();
 	}
 
 	return NewSection;
@@ -116,7 +121,12 @@ FMovieSceneEvalTemplatePtr UMovieSceneAnimatorTrack::CreateTemplateForSection(co
 		return FMovieSceneEvalTemplatePtr();
 	}
 
-	return FMovieSceneAnimatorEvalTemplate(AnimatorSection->GetChannel());
+	FMovieSceneAnimatorSectionData SectionData;
+	SectionData.Channel = AnimatorSection->GetChannel();
+	SectionData.bUseSectionTime = AnimatorSection->GetUseSectionTime();
+	SectionData.Section = AnimatorSection;
+
+	return FMovieSceneAnimatorEvalTemplate(SectionData);
 }
 
 #undef LOCTEXT_NAMESPACE
diff --git a/Engine/Plugins/Experimental/PropertyAnimatorCore/Source/PropertyAnimatorCore/Private/TimeSources/PropertyAnimatorCoreManualTimeSource.cpp b/Engine/Plugins/Experimental/PropertyAnimatorCore/Source/PropertyAnimatorCore/Private/TimeSources/PropertyAnimatorCoreManualTimeSource.cpp
--- a/Engine/Plugins/Experimental/PropertyAnimatorCore/Source/PropertyAnimatorCore/Private/TimeSources/PropertyAnimatorCoreManualTimeSource.cpp
+++ b/Engine/Plugins/Experimental/PropertyAnimatorCore/Source/PropertyAnimatorCore/Private/TimeSources/PropertyAnimatorCoreManualTimeSource.cpp
@@ -36,7 +36,7 @@ void UPropertyAnimatorCoreManualTimeSource::SetState(const FPropertyAnimatorCore
 	OnStateChanged();
 }
 
-double UPropertyAnimatorCoreManualTimeSource::GetTimeElapsed()
+bool UPropertyAnimatorCoreManualTimeSource::UpdateEvaluationData(FPropertyAnimatorCoreTimeSourceEvaluationData& OutData)
 {
 	if (!bOverrideTime)
 	{
@@ -52,19 +52,15 @@ double UPropertyAnimatorCoreManualTimeSource::GetTimeElapsed()
 		{
 			CustomTime -= FApp::GetDeltaTime();
 		}
-	}
 
-	return CustomTime;
-}
+		OutData.TimeElapsed = CustomTime;
 
-bool UPropertyAnimatorCoreManualTimeSource::IsTimeSourceReady() const
-{
-	if (bOverrideTime)
-	{
-		return true;
+		return ActiveStatus != EPropertyAnimatorCoreManualStatus::Stopped;
 	}
 
-	return ActiveStatus != EPropertyAnimatorCoreManualStatus::Stopped;
+	OutData.TimeElapsed = CustomTime;
+
+	return true;
 }
 
 void UPropertyAnimatorCoreManualTimeSource::OnTimeSourceActive()
diff --git a/Engine/Plugins/Experimental/PropertyAnimatorCore/Source/PropertyAnimatorCore/Private/TimeSources/PropertyAnimatorCoreSequencerTimeSource.cpp b/Engine/Plugins/Experimental/PropertyAnimatorCore/Source/PropertyAnimatorCore/Private/TimeSources/PropertyAnimatorCoreSequencerTimeSource.cpp
--- a/Engine/Plugins/Experimental/PropertyAnimatorCore/Source/PropertyAnimatorCore/Private/TimeSources/PropertyAnimatorCoreSequencerTimeSource.cpp
+++ b/Engine/Plugins/Experimental/PropertyAnimatorCore/Source/PropertyAnimatorCore/Private/TimeSources/PropertyAnimatorCoreSequencerTimeSource.cpp
@@ -4,13 +4,16 @@
 
 UPropertyAnimatorCoreSequencerTimeSource::FOnAnimatorTimeEvaluated UPropertyAnimatorCoreSequencerTimeSource::OnAnimatorTimeEvaluated;
 
-double UPropertyAnimatorCoreSequencerTimeSource::GetTimeElapsed()
+bool UPropertyAnimatorCoreSequencerTimeSource::UpdateEvaluationData(FPropertyAnimatorCoreTimeSourceEvaluationData& OutData)
 {
-	return EvalTime.Get(0);
-}
+	if (!EvalTime.IsSet() || !EvalMagnitude.IsSet())
+	{
+		return false;
+	}
+
+	OutData.TimeElapsed = EvalTime.GetValue();
+	OutData.Magnitude = EvalMagnitude.GetValue();
 
-bool UPropertyAnimatorCoreSequencerTimeSource::IsTimeSourceReady() const
-{
 	return true;
 }
 
@@ -33,10 +36,11 @@ void UPropertyAnimatorCoreSequencerTimeSource::SetChannel(uint8 InChannel)
 	ChannelData.Channel = InChannel;
 }
 
-void UPropertyAnimatorCoreSequencerTimeSource::OnSequencerTimeEvaluated(uint8 InChannel, double InTimeEval)
+void UPropertyAnimatorCoreSequencerTimeSource::OnSequencerTimeEvaluated(uint8 InChannel, const TOptional<double>& InTimeEval, const TOptional<float>& InMagnitudeEval)
 {
 	if (ChannelData.Channel == InChannel)
 	{
 		EvalTime = InTimeEval;
+		EvalMagnitude = InMagnitudeEval;
 	}
 }
diff --git a/Engine/Plugins/Experimental/PropertyAnimatorCore/Source/PropertyAnimatorCore/Private/TimeSources/PropertyAnimatorCoreSystemTimeSource.cpp b/Engine/Plugins/Experimental/PropertyAnimatorCore/Source/PropertyAnimatorCore/Private/TimeSources/PropertyAnimatorCoreSystemTimeSource.cpp
--- a/Engine/Plugins/Experimental/PropertyAnimatorCore/Source/PropertyAnimatorCore/Private/TimeSources/PropertyAnimatorCoreSystemTimeSource.cpp
+++ b/Engine/Plugins/Experimental/PropertyAnimatorCore/Source/PropertyAnimatorCore/Private/TimeSources/PropertyAnimatorCoreSystemTimeSource.cpp
@@ -5,39 +5,32 @@
 #include "Internationalization/Regex.h"
 #include "Misc/DateTime.h"
 
-double UPropertyAnimatorCoreSystemTimeSource::GetTimeElapsed()
+bool UPropertyAnimatorCoreSystemTimeSource::UpdateEvaluationData(FPropertyAnimatorCoreTimeSourceEvaluationData& OutData)
 {
-	double TimeElapsedSeconds = 0;
-
 	switch(Mode)
 	{
 	case EPropertyAnimatorCoreSystemMode::LocalTime:
 		{
-			TimeElapsedSeconds = (FDateTime::Now() - FDateTime::MinValue()).GetTotalSeconds();
+			OutData.TimeElapsed = (FDateTime::Now() - FDateTime::MinValue()).GetTotalSeconds();
 		}
 		break;
 	case EPropertyAnimatorCoreSystemMode::UtcTime:
 		{
-			TimeElapsedSeconds = (FDateTime::UtcNow() - FDateTime::MinValue()).GetTotalSeconds();
+			OutData.TimeElapsed = (FDateTime::UtcNow() - FDateTime::MinValue()).GetTotalSeconds();
 		}
 		break;
 	case EPropertyAnimatorCoreSystemMode::Countdown:
 		{
-			TimeElapsedSeconds = (CountdownTimeSpan - (FDateTime::Now() - ActivationTime)).GetTotalSeconds();
+			OutData.TimeElapsed = (CountdownTimeSpan - (FDateTime::Now() - ActivationTime)).GetTotalSeconds();
 		}
 		break;
 	case EPropertyAnimatorCoreSystemMode::Stopwatch:
 		{
-			TimeElapsedSeconds = (FDateTime::Now() - ActivationTime).GetTotalSeconds();
+			OutData.TimeElapsed = (FDateTime::Now() - ActivationTime).GetTotalSeconds();
 		}
 		break;
 	}
 
-	return TimeElapsedSeconds;
-}
-
-bool UPropertyAnimatorCoreSystemTimeSource::IsTimeSourceReady() const
-{
 	return true;
 }
 
diff --git a/Engine/Plugins/Experimental/PropertyAnimatorCore/Source/PropertyAnimatorCore/Private/TimeSources/PropertyAnimatorCoreTimeSourceBase.cpp b/Engine/Plugins/Experimental/PropertyAnimatorCore/Source/PropertyAnimatorCore/Private/TimeSources/PropertyAnimatorCoreTimeSourceBase.cpp
--- a/Engine/Plugins/Experimental/PropertyAnimatorCore/Source/PropertyAnimatorCore/Private/TimeSources/PropertyAnimatorCoreTimeSourceBase.cpp
+++ b/Engine/Plugins/Experimental/PropertyAnimatorCore/Source/PropertyAnimatorCore/Private/TimeSources/PropertyAnimatorCoreTimeSourceBase.cpp
@@ -24,23 +24,23 @@ void UPropertyAnimatorCoreTimeSourceBase::DeactivateTimeSource()
 	OnTimeSourceInactive();
 }
 
-TOptional<double> UPropertyAnimatorCoreTimeSourceBase::GetConditionalTimeElapsed()
+EPropertyAnimatorCoreTimeSourceResult UPropertyAnimatorCoreTimeSourceBase::FetchEvaluationData(FPropertyAnimatorCoreTimeSourceEvaluationData& OutEvaluationData)
 {
-	if (!IsTimeSourceReady())
+	if (!UpdateEvaluationData(OutEvaluationData))
 	{
-		return TOptional<double>();
+		// Reset evaluation state
+		return EPropertyAnimatorCoreTimeSourceResult::Reset;
 	}
 
-	double NewTimeElapsed = GetTimeElapsed();
-
-	if (!IsValidTimeElapsed(NewTimeElapsed))
+	if (!IsFramerateAllowed(OutEvaluationData.TimeElapsed))
 	{
-		return TOptional<double>();
+		// Skip evaluation for this run
+		return EPropertyAnimatorCoreTimeSourceResult::Skip;
 	}
 
-	LastTimeElapsed = NewTimeElapsed;
+	LastTimeElapsed = OutEvaluationData.TimeElapsed;
 
-	return LastTimeElapsed;
+	return EPropertyAnimatorCoreTimeSourceResult::Evaluate;
 }
 
 void UPropertyAnimatorCoreTimeSourceBase::SetFrameRate(float InFrameRate)
@@ -53,7 +53,12 @@ void UPropertyAnimatorCoreTimeSourceBase::SetUseFrameRate(bool bInUseFrameRate)
 	bUseFrameRate = bInUseFrameRate;
 }
 
-bool UPropertyAnimatorCoreTimeSourceBase::IsValidTimeElapsed(double InTimeElapsed) const
+bool UPropertyAnimatorCoreTimeSourceBase::UpdateEvaluationData(FPropertyAnimatorCoreTimeSourceEvaluationData& OutData)
+{
+	return false;
+}
+
+bool UPropertyAnimatorCoreTimeSourceBase::IsFramerateAllowed(double InNewTime) const
 {
-	return !bUseFrameRate || FMath::IsNearlyZero(FrameRate) || FMath::Abs(InTimeElapsed - LastTimeElapsed) > FMath::Abs(1.f / FrameRate);
+	return !bUseFrameRate || FMath::IsNearlyZero(FrameRate) || FMath::Abs(InNewTime - LastTimeElapsed) > FMath::Abs(1.f / FrameRate);
 }
diff --git a/Engine/Plugins/Experimental/PropertyAnimatorCore/Source/PropertyAnimatorCore/Private/TimeSources/PropertyAnimatorCoreWorldTimeSource.cpp b/Engine/Plugins/Experimental/PropertyAnimatorCore/Source/PropertyAnimatorCore/Private/TimeSources/PropertyAnimatorCoreWorldTimeSource.cpp
--- a/Engine/Plugins/Experimental/PropertyAnimatorCore/Source/PropertyAnimatorCore/Private/TimeSources/PropertyAnimatorCoreWorldTimeSource.cpp
+++ b/Engine/Plugins/Experimental/PropertyAnimatorCore/Source/PropertyAnimatorCore/Private/TimeSources/PropertyAnimatorCoreWorldTimeSource.cpp
@@ -2,20 +2,9 @@
 
 #include "TimeSources/PropertyAnimatorCoreWorldTimeSource.h"
 
-#include "Animators/PropertyAnimatorCoreBase.h"
 #include "Engine/World.h"
 
-double UPropertyAnimatorCoreWorldTimeSource::GetTimeElapsed()
-{
-	if (const UWorld* World = GetWorld())
-	{
-		return World->GetTimeSeconds();
-	}
-
-	return 0.f;
-}
-
-bool UPropertyAnimatorCoreWorldTimeSource::IsTimeSourceReady() const
+bool UPropertyAnimatorCoreWorldTimeSource::UpdateEvaluationData(FPropertyAnimatorCoreTimeSourceEvaluationData& OutData)
 {
 	const UWorld* World = GetWorld();
 
@@ -24,6 +13,8 @@ bool UPropertyAnimatorCoreWorldTimeSource::IsTimeSourceReady() const
 		return false;
 	}
 
+	OutData.TimeElapsed = World->GetTimeSeconds();
+
 	return World->IsEditorWorld()
 		|| World->IsPreviewWorld()
 		|| (World->IsGameWorld() && World->HasBegunPlay());
diff --git a/Engine/Plugins/Experimental/PropertyAnimatorCore/Source/PropertyAnimatorCore/Public/Components/PropertyAnimatorCoreComponent.h b/Engine/Plugins/Experimental/PropertyAnimatorCore/Source/PropertyAnimatorCore/Public/Components/PropertyAnimatorCoreComponent.h
--- a/Engine/Plugins/Experimental/PropertyAnimatorCore/Source/PropertyAnimatorCore/Public/Components/PropertyAnimatorCoreComponent.h
+++ b/Engine/Plugins/Experimental/PropertyAnimatorCore/Source/PropertyAnimatorCore/Public/Components/PropertyAnimatorCoreComponent.h
@@ -134,11 +134,11 @@ protected:
 	bool bAnimatorsEnabled = true;
 
 	/** Global magnitude for all animators controlled by this component */
-	UPROPERTY(EditInstanceOnly, Getter, Setter, Category="Animator", meta=(ClampMin="0", ClampMax="1", UIMin="0", UIMax="1", AllowPrivateAccess="true"))
+	UPROPERTY(EditInstanceOnly, Getter, Setter, DisplayName="Global Magnitude", Category="Animator", meta=(ClampMin="0", ClampMax="1", UIMin="0", UIMax="1", AllowPrivateAccess="true"))
 	float AnimatorsMagnitude = 1.f;
 
 	/** The global time source to use, can be overriden in animator */
-	UPROPERTY(EditInstanceOnly, Setter, Getter, Category="Animator", meta=(GetOptions="GetTimeSourceNames"))
+	UPROPERTY(EditInstanceOnly, Setter, Getter, DisplayName="Global Time Source Name", Category="Animator", meta=(GetOptions="GetTimeSourceNames"))
 	FName AnimatorsTimeSourceName = NAME_None;
 
 	/** Active time source with its options, determined by its name */
diff --git a/Engine/Plugins/Experimental/PropertyAnimatorCore/Source/PropertyAnimatorCore/Public/Sequencer/MovieSceneAnimatorSection.h b/Engine/Plugins/Experimental/PropertyAnimatorCore/Source/PropertyAnimatorCore/Public/Sequencer/MovieSceneAnimatorSection.h
--- a/Engine/Plugins/Experimental/PropertyAnimatorCore/Source/PropertyAnimatorCore/Public/Sequencer/MovieSceneAnimatorSection.h
+++ b/Engine/Plugins/Experimental/PropertyAnimatorCore/Source/PropertyAnimatorCore/Public/Sequencer/MovieSceneAnimatorSection.h
@@ -2,6 +2,7 @@
 
 #pragma once
 
+#include "MovieSceneAnimatorTypes.h"
 #include "MovieSceneSection.h"
 #include "MovieSceneAnimatorSection.generated.h"
 
@@ -15,13 +16,25 @@ public:
 	UMovieSceneAnimatorSection();
 
 	PROPERTYANIMATORCORE_API void SetChannel(uint8 InChannel);
+
 	uint8 GetChannel() const
 	{
 		return Channel;
 	}
 
+	PROPERTYANIMATORCORE_API void SetUseSectionTime(bool bInUse);
+
+	bool GetUseSectionTime() const
+	{
+		return bUseSectionTime;
+	}
+
 protected:
 	/** Channel used to push sequencer time to */
 	UPROPERTY(EditAnywhere, Category="Animator")
 	uint8 Channel = 0;
+
+	/** Whether to use the section time (playback range) or track time (sequencer time) */
+	UPROPERTY(EditAnywhere, Category="Animator")
+	bool bUseSectionTime = true;
 };
diff --git a/Engine/Plugins/Experimental/PropertyAnimatorCore/Source/PropertyAnimatorCore/Public/Sequencer/MovieSceneAnimatorTypes.h b/Engine/Plugins/Experimental/PropertyAnimatorCore/Source/PropertyAnimatorCore/Public/Sequencer/MovieSceneAnimatorTypes.h
new file mode 100644
--- /dev/null
+++ b/Engine/Plugins/Experimental/PropertyAnimatorCore/Source/PropertyAnimatorCore/Public/Sequencer/MovieSceneAnimatorTypes.h
@@ -0,0 +1,19 @@
+// Copyright Epic Games, Inc. All Rights Reserved.
+
+#pragma once
+
+#include "HAL/Platform.h"
+
+class UMovieSceneAnimatorSection;
+
+struct FMovieSceneAnimatorSectionData
+{
+	/** Channel used to push sequencer time to */
+	uint8 Channel = 0;
+
+	/** Whether to use the section time or track time (sequencer time) */
+	bool bUseSectionTime = true;
+
+	/** The section object to get easing from */
+	const UMovieSceneAnimatorSection* Section = nullptr;
+};
diff --git a/Engine/Plugins/Experimental/PropertyAnimatorCore/Source/PropertyAnimatorCore/Public/TimeSources/PropertyAnimatorCoreManualTimeSource.h b/Engine/Plugins/Experimental/PropertyAnimatorCore/Source/PropertyAnimatorCore/Public/TimeSources/PropertyAnimatorCoreManualTimeSource.h
--- a/Engine/Plugins/Experimental/PropertyAnimatorCore/Source/PropertyAnimatorCore/Public/TimeSources/PropertyAnimatorCoreManualTimeSource.h
+++ b/Engine/Plugins/Experimental/PropertyAnimatorCore/Source/PropertyAnimatorCore/Public/TimeSources/PropertyAnimatorCoreManualTimeSource.h
@@ -57,8 +57,7 @@ public:
 	}
 
 	//~ Begin UPropertyAnimatorTimeSourceBase
-	virtual double GetTimeElapsed() override;
-	virtual bool IsTimeSourceReady() const override;
+	virtual bool UpdateEvaluationData(FPropertyAnimatorCoreTimeSourceEvaluationData& OutData) override;
 	virtual void OnTimeSourceActive() override;
 	virtual void OnTimeSourceInactive() override;
 	//~ End UPropertyAnimatorTimeSourceBase
diff --git a/Engine/Plugins/Experimental/PropertyAnimatorCore/Source/PropertyAnimatorCore/Public/TimeSources/PropertyAnimatorCoreSequencerTimeSource.h b/Engine/Plugins/Experimental/PropertyAnimatorCore/Source/PropertyAnimatorCore/Public/TimeSources/PropertyAnimatorCoreSequencerTimeSource.h
--- a/Engine/Plugins/Experimental/PropertyAnimatorCore/Source/PropertyAnimatorCore/Public/TimeSources/PropertyAnimatorCoreSequencerTimeSource.h
+++ b/Engine/Plugins/Experimental/PropertyAnimatorCore/Source/PropertyAnimatorCore/Public/TimeSources/PropertyAnimatorCoreSequencerTimeSource.h
@@ -21,7 +21,7 @@ class UPropertyAnimatorCoreSequencerTimeSource : public UPropertyAnimatorCoreTim
 	GENERATED_BODY()
 
 public:
-	DECLARE_MULTICAST_DELEGATE_TwoParams(FOnAnimatorTimeEvaluated, uint8 /** Channel */, double /** EvalTime */)
+	DECLARE_MULTICAST_DELEGATE_ThreeParams(FOnAnimatorTimeEvaluated, uint8 /** Channel */, const TOptional<double>& /** EvalTime */, const TOptional<float>& /** EvalMagnitude */)
 	static FOnAnimatorTimeEvaluated OnAnimatorTimeEvaluated;
 
 	UPropertyAnimatorCoreSequencerTimeSource()
@@ -29,8 +29,7 @@ public:
 	{}
 
 	//~ Begin UPropertyAnimatorTimeSourceBase
-	virtual double GetTimeElapsed() override;
-	virtual bool IsTimeSourceReady() const override;
+	virtual bool UpdateEvaluationData(FPropertyAnimatorCoreTimeSourceEvaluationData& OutData) override;
 	virtual void OnTimeSourceActive() override;
 	virtual void OnTimeSourceInactive() override;
 	//~ End UPropertyAnimatorTimeSourceBase
@@ -42,7 +41,7 @@ public:
 	}
 
 protected:
-	void OnSequencerTimeEvaluated(uint8 InChannel, double InTimeEval);
+	void OnSequencerTimeEvaluated(uint8 InChannel, const TOptional<double>& InTimeEval, const TOptional<float>& InMagnitudeEval);
 
 	/** Channel to sample time from */
 	UPROPERTY(EditInstanceOnly, DisplayName="Channel", Category="Animator")
@@ -50,4 +49,7 @@ protected:
 
 	/** Last evaluated time received */
 	TOptional<double> EvalTime;
+
+	/** Last evaluated magnitude received */
+	TOptional<float> EvalMagnitude;
 };
\ No newline at end of file
diff --git a/Engine/Plugins/Experimental/PropertyAnimatorCore/Source/PropertyAnimatorCore/Public/TimeSources/PropertyAnimatorCoreSystemTimeSource.h b/Engine/Plugins/Experimental/PropertyAnimatorCore/Source/PropertyAnimatorCore/Public/TimeSources/PropertyAnimatorCoreSystemTimeSource.h
--- a/Engine/Plugins/Experimental/PropertyAnimatorCore/Source/PropertyAnimatorCore/Public/TimeSources/PropertyAnimatorCoreSystemTimeSource.h
+++ b/Engine/Plugins/Experimental/PropertyAnimatorCore/Source/PropertyAnimatorCore/Public/TimeSources/PropertyAnimatorCoreSystemTimeSource.h
@@ -33,8 +33,7 @@ public:
 	{}
 
 	//~ Begin UPropertyAnimatorTimeSourceBase
-	virtual double GetTimeElapsed() override;
-	virtual bool IsTimeSourceReady() const override;
+	virtual bool UpdateEvaluationData(FPropertyAnimatorCoreTimeSourceEvaluationData& OutData) override;
 	virtual void OnTimeSourceActive() override;
 	//~ End UPropertyAnimatorTimeSourceBase
 
diff --git a/Engine/Plugins/Experimental/PropertyAnimatorCore/Source/PropertyAnimatorCore/Public/TimeSources/PropertyAnimatorCoreTimeSourceBase.h b/Engine/Plugins/Experimental/PropertyAnimatorCore/Source/PropertyAnimatorCore/Public/TimeSources/PropertyAnimatorCoreTimeSourceBase.h
--- a/Engine/Plugins/Experimental/PropertyAnimatorCore/Source/PropertyAnimatorCore/Public/TimeSources/PropertyAnimatorCoreTimeSourceBase.h
+++ b/Engine/Plugins/Experimental/PropertyAnimatorCore/Source/PropertyAnimatorCore/Public/TimeSources/PropertyAnimatorCoreTimeSourceBase.h
@@ -6,6 +6,27 @@
 
 class UPropertyAnimatorCoreBase;
 
+/** Enumerates all possible outcomes for the time source */
+enum class EPropertyAnimatorCoreTimeSourceResult
+{
+	/** Time is the same as previous, evaluation can be skipped */
+	Skip,
+	/** Time is in an invalid state or out of range, restore values */
+	Reset,
+	/** Time is valid and in range, evaluate time */
+	Evaluate
+};
+
+/** Stores all the data used by animators during evaluation */
+struct FPropertyAnimatorCoreTimeSourceEvaluationData
+{
+	/** Time elapsed for animators evaluation */
+	double TimeElapsed = 0.0;
+
+	/** Time magnitude for animators to fade in/out effect based on time */
+	float Magnitude = 1.f;
+};
+
 /**
  * Abstract base class for time source used by property animators
  * Can be transient or saved to disk if contains user set data
@@ -34,7 +55,7 @@ public:
 		return bTimeSourceActive;
 	}
 
-	TOptional<double> GetConditionalTimeElapsed();
+	EPropertyAnimatorCoreTimeSourceResult FetchEvaluationData(FPropertyAnimatorCoreTimeSourceEvaluationData& OutEvaluationData);
 
 	FName GetTimeSourceName() const
 	{
@@ -53,21 +74,14 @@ public:
 		return bUseFrameRate;
 	}
 
-protected:
-	/** Returns the time elapsed for animators */
-	virtual double GetTimeElapsed()
-	{
-		return 0;
-	}
-
-	/** Checks if this time source is ready to be used by the animator */
-	virtual bool IsTimeSourceReady() const
+	double GetLastTimeElapsed() const
 	{
-		return false;
+		return LastTimeElapsed;
 	}
 
-	/** Check if the time elapsed is valid based on the context */
-	PROPERTYANIMATORCORE_API virtual bool IsValidTimeElapsed(double InTimeElapsed) const;
+protected:
+	/** Retrieve evaluation data to provide animators, return true if data is valid, false otherwise */
+	PROPERTYANIMATORCORE_API virtual bool UpdateEvaluationData(FPropertyAnimatorCoreTimeSourceEvaluationData& OutData);
 
 	/** Time source CDO is registered by subsystem */
 	virtual void OnTimeSourceRegistered() {}
@@ -82,6 +96,8 @@ protected:
 	virtual void OnTimeSourceInactive() {}
 
 private:
+	bool IsFramerateAllowed(double InNewTime) const;
+
 	/** Use a specific framerate */
 	UPROPERTY(EditInstanceOnly, Setter="SetUseFrameRate", Getter="GetUseFrameRate", Category="Animator", meta=(InlineEditConditionToggle))
 	bool bUseFrameRate = false;
diff --git a/Engine/Plugins/Experimental/PropertyAnimatorCore/Source/PropertyAnimatorCore/Public/TimeSources/PropertyAnimatorCoreWorldTimeSource.h b/Engine/Plugins/Experimental/PropertyAnimatorCore/Source/PropertyAnimatorCore/Public/TimeSources/PropertyAnimatorCoreWorldTimeSource.h
--- a/Engine/Plugins/Experimental/PropertyAnimatorCore/Source/PropertyAnimatorCore/Public/TimeSources/PropertyAnimatorCoreWorldTimeSource.h
+++ b/Engine/Plugins/Experimental/PropertyAnimatorCore/Source/PropertyAnimatorCore/Public/TimeSources/PropertyAnimatorCoreWorldTimeSource.h
@@ -20,7 +20,6 @@ public:
 	{}
 
 	//~ Begin UPropertyAnimatorTimeSourceBase
-	virtual double GetTimeElapsed() override;
-	virtual bool IsTimeSourceReady() const override;
+	virtual bool UpdateEvaluationData(FPropertyAnimatorCoreTimeSourceEvaluationData& OutData) override;
 	//~ End UPropertyAnimatorTimeSourceBase
 };
\ No newline at end of file
diff --git a/Engine/Plugins/Experimental/PropertyAnimatorCore/Source/PropertyAnimatorCoreEditor/Private/Sequencer/MovieSceneAnimatorTrackEditor.cpp b/Engine/Plugins/Experimental/PropertyAnimatorCore/Source/PropertyAnimatorCoreEditor/Private/Sequencer/MovieSceneAnimatorTrackEditor.cpp
--- a/Engine/Plugins/Experimental/PropertyAnimatorCore/Source/PropertyAnimatorCoreEditor/Private/Sequencer/MovieSceneAnimatorTrackEditor.cpp
+++ b/Engine/Plugins/Experimental/PropertyAnimatorCore/Source/PropertyAnimatorCoreEditor/Private/Sequencer/MovieSceneAnimatorTrackEditor.cpp
@@ -30,6 +30,12 @@ void FMovieSceneAnimatorTrackEditor::BuildAddTrackMenu(FMenuBuilder& InMenuBuild
 	);
 }
 
+TSharedPtr<SWidget> FMovieSceneAnimatorTrackEditor::BuildOutlinerEditWidget(const FGuid& InObjectBinding, UMovieSceneTrack* InTrack, const FBuildEditWidgetParams& InParams)
+{
+	// Empty, do not allow the creation of new sections, only one to rule them all
+	return nullptr;
+}
+
 void FMovieSceneAnimatorTrackEditor::BindDelegates()
 {
 	OnAddAnimatorTrack.AddSP(this, &FMovieSceneAnimatorTrackEditor::ExecuteAddTrack);
diff --git a/Engine/Plugins/Experimental/PropertyAnimatorCore/Source/PropertyAnimatorCoreEditor/Private/Sequencer/MovieSceneAnimatorTrackEditor.h b/Engine/Plugins/Experimental/PropertyAnimatorCore/Source/PropertyAnimatorCoreEditor/Private/Sequencer/MovieSceneAnimatorTrackEditor.h
--- a/Engine/Plugins/Experimental/PropertyAnimatorCore/Source/PropertyAnimatorCoreEditor/Private/Sequencer/MovieSceneAnimatorTrackEditor.h
+++ b/Engine/Plugins/Experimental/PropertyAnimatorCore/Source/PropertyAnimatorCoreEditor/Private/Sequencer/MovieSceneAnimatorTrackEditor.h
@@ -31,6 +31,7 @@ public:
 private:
 	//~ Begin FMovieSceneTrackEditor
 	virtual void BuildAddTrackMenu(FMenuBuilder& InMenuBuilder) override;
+	virtual TSharedPtr<SWidget> BuildOutlinerEditWidget(const FGuid& InObjectBinding, UMovieSceneTrack* InTrack, const FBuildEditWidgetParams& InParams) override;
 	//~ End FMovieSceneTrackEditor
 
 	void BindDelegates();
