MotionDesign : PropertyAnimator
- Disallow adding an animator track within a UNiagaraSequence

diff --git a/Engine/Plugins/Experimental/PropertyAnimatorCore/Source/PropertyAnimatorCoreEditor/Private/Sequencer/MovieSceneAnimatorTrackEditor.cpp b/Engine/Plugins/Experimental/PropertyAnimatorCore/Source/PropertyAnimatorCoreEditor/Private/Sequencer/MovieSceneAnimatorTrackEditor.cpp
--- a/Engine/Plugins/Experimental/PropertyAnimatorCore/Source/PropertyAnimatorCoreEditor/Private/Sequencer/MovieSceneAnimatorTrackEditor.cpp
+++ b/Engine/Plugins/Experimental/PropertyAnimatorCore/Source/PropertyAnimatorCoreEditor/Private/Sequencer/MovieSceneAnimatorTrackEditor.cpp
@@ -25,7 +25,8 @@ void FMovieSceneAnimatorTrackEditor::BuildAddTrackMenu(FMenuBuilder& InMenuBuild
 		LOCTEXT("AddAnimatorTrack.Tooltip", "Adds a new track that uses the time of the current sequence to drive animators."),
 		FSlateIconFinder::FindIconForClass(UMovieSceneAnimatorTrack::StaticClass()),
 		FUIAction(
-			FExecuteAction::CreateSP(this, &FMovieSceneAnimatorTrackEditor::ExecuteAddTrack, Channel)
+			FExecuteAction::CreateSP(this, &FMovieSceneAnimatorTrackEditor::ExecuteAddTrack, Channel),
+			FCanExecuteAction::CreateSP(this, &FMovieSceneAnimatorTrackEditor::CanExecuteAddTrack)
 		)
 	);
 }
@@ -36,6 +37,19 @@ TSharedPtr<SWidget> FMovieSceneAnimatorTrackEditor::BuildOutlinerEditWidget(cons
 	return nullptr;
 }
 
+bool FMovieSceneAnimatorTrackEditor::SupportsSequence(UMovieSceneSequence* InSequence) const
+{
+	const ETrackSupport TrackSupported = InSequence ? InSequence->IsTrackSupported(UMovieSceneAnimatorTrack::StaticClass()) : ETrackSupport::NotSupported;
+
+	if (TrackSupported == ETrackSupport::NotSupported)
+	{
+		return false;
+	}
+
+	// UNiagaraSequence is private so check against class name
+	return InSequence->GetClass()->GetFName() != TEXT("NiagaraSequence");
+}
+
 void FMovieSceneAnimatorTrackEditor::BindDelegates()
 {
 	OnAddAnimatorTrack.AddSP(this, &FMovieSceneAnimatorTrackEditor::ExecuteAddTrack);
@@ -59,6 +73,11 @@ void FMovieSceneAnimatorTrackEditor::GetTrackCount(uint8 InChannel, int32& OutCo
 	}
 }
 
+bool FMovieSceneAnimatorTrackEditor::CanExecuteAddTrack() const
+{
+	return GetSequencer() && GetFocusedMovieScene();
+}
+
 void FMovieSceneAnimatorTrackEditor::ExecuteAddTrack(uint8 InChannel)
 {
 	UMovieScene* FocusedMovieScene = GetFocusedMovieScene();
diff --git a/Engine/Plugins/Experimental/PropertyAnimatorCore/Source/PropertyAnimatorCoreEditor/Private/Sequencer/MovieSceneAnimatorTrackEditor.h b/Engine/Plugins/Experimental/PropertyAnimatorCore/Source/PropertyAnimatorCoreEditor/Private/Sequencer/MovieSceneAnimatorTrackEditor.h
--- a/Engine/Plugins/Experimental/PropertyAnimatorCore/Source/PropertyAnimatorCoreEditor/Private/Sequencer/MovieSceneAnimatorTrackEditor.h
+++ b/Engine/Plugins/Experimental/PropertyAnimatorCore/Source/PropertyAnimatorCoreEditor/Private/Sequencer/MovieSceneAnimatorTrackEditor.h
@@ -5,6 +5,8 @@
 #include "KeyframeTrackEditor.h"
 #include "Sequencer/MovieSceneAnimatorTrack.h"
 
+class UMovieSceneSequence;
+
 /** Animator track editor to add animator track and section */
 class FMovieSceneAnimatorTrackEditor : public FKeyframeTrackEditor<UMovieSceneAnimatorTrack>
 {
@@ -32,11 +34,13 @@ private:
 	//~ Begin FMovieSceneTrackEditor
 	virtual void BuildAddTrackMenu(FMenuBuilder& InMenuBuilder) override;
 	virtual TSharedPtr<SWidget> BuildOutlinerEditWidget(const FGuid& InObjectBinding, UMovieSceneTrack* InTrack, const FBuildEditWidgetParams& InParams) override;
+	virtual bool SupportsSequence(UMovieSceneSequence* InSequence) const override;
 	//~ End FMovieSceneTrackEditor
 
 	void BindDelegates();
 
 	void GetTrackCount(uint8 InChannel, int32& OutCount) const;
 
+	bool CanExecuteAddTrack() const;
 	void ExecuteAddTrack(uint8 InChannel);
 };
