MotionDesign : PropertyAnimator
- Disallow adding an animator track within a sequence that is not a ULevelSequence

diff --git a/Engine/Plugins/Experimental/PropertyAnimatorCore/Source/PropertyAnimatorCoreEditor/Private/Sequencer/MovieSceneAnimatorTrackEditor.cpp b/Engine/Plugins/Experimental/PropertyAnimatorCore/Source/PropertyAnimatorCoreEditor/Private/Sequencer/MovieSceneAnimatorTrackEditor.cpp
--- a/Engine/Plugins/Experimental/PropertyAnimatorCore/Source/PropertyAnimatorCoreEditor/Private/Sequencer/MovieSceneAnimatorTrackEditor.cpp
+++ b/Engine/Plugins/Experimental/PropertyAnimatorCore/Source/PropertyAnimatorCoreEditor/Private/Sequencer/MovieSceneAnimatorTrackEditor.cpp
@@ -2,6 +2,7 @@
 
 #include "MovieSceneAnimatorTrackEditor.h"
 
+#include "LevelSequence.h"
 #include "Sequencer/MovieSceneAnimatorSection.h"
 #include "Styling/SlateIconFinder.h"
 
@@ -46,8 +47,7 @@ bool FMovieSceneAnimatorTrackEditor::SupportsSequence(UMovieSceneSequence* InSeq
 		return false;
 	}
 
-	// UNiagaraSequence is private so check against class name
-	return InSequence->GetClass()->GetFName() != TEXT("NiagaraSequence");
+	return InSequence->IsA(ULevelSequence::StaticClass());
 }
 
 void FMovieSceneAnimatorTrackEditor::BindDelegates()
diff --git a/Engine/Plugins/Experimental/PropertyAnimatorCore/Source/PropertyAnimatorCoreEditor/PropertyAnimatorCoreEditor.Build.cs b/Engine/Plugins/Experimental/PropertyAnimatorCore/Source/PropertyAnimatorCoreEditor/PropertyAnimatorCoreEditor.Build.cs
--- a/Engine/Plugins/Experimental/PropertyAnimatorCore/Source/PropertyAnimatorCoreEditor/PropertyAnimatorCoreEditor.Build.cs
+++ b/Engine/Plugins/Experimental/PropertyAnimatorCore/Source/PropertyAnimatorCoreEditor/PropertyAnimatorCoreEditor.Build.cs
@@ -24,6 +24,7 @@ public class PropertyAnimatorCoreEditor : ModuleRules
 		        "EditorWidgets",
 		        "Engine",
 		        "InputCore",
+				"LevelSequence",
 				"MovieScene",
 				"MovieSceneTools",
 		        "OperatorStackEditor",
