MotionDesign : PropertyAnimatorCore
- Share animator metadata across all instances of the same class (name, category)
- Simple add menu will now group animators based on their category
- Adjusting animators with new changes

diff --git a/Engine/Plugins/Experimental/Avalanche/Source/AvalanchePropertyAnimatorEditor/Private/Outliner/AvaPropertyAnimatorEditorOutliner.cpp b/Engine/Plugins/Experimental/Avalanche/Source/AvalanchePropertyAnimatorEditor/Private/Outliner/AvaPropertyAnimatorEditorOutliner.cpp
--- a/Engine/Plugins/Experimental/Avalanche/Source/AvalanchePropertyAnimatorEditor/Private/Outliner/AvaPropertyAnimatorEditorOutliner.cpp
+++ b/Engine/Plugins/Experimental/Avalanche/Source/AvalanchePropertyAnimatorEditor/Private/Outliner/AvaPropertyAnimatorEditorOutliner.cpp
@@ -12,7 +12,7 @@ FAvaPropertyAnimatorEditorOutliner::FAvaPropertyAnimatorEditorOutliner(IAvaOutli
 	: FAvaOutlinerObject(InOutliner, InAnimator)
 	, PropertyAnimator(InAnimator)
 {
-	ItemName = FText::FromString(PropertyAnimator->GetAnimatorDisplayName());
+	ItemName = FText::FromName(PropertyAnimator->GetAnimatorDisplayName());
 	ItemIcon = FSlateIconFinder::FindIconForClass(UPropertyAnimatorCoreComponent::StaticClass());
 	ItemTooltip = FText::FromName(PropertyAnimator->GetAnimatorOriginalName());
 }
diff --git a/Engine/Plugins/Experimental/PropertyAnimator/Source/PropertyAnimator/Private/Animators/PropertyAnimatorBounce.cpp b/Engine/Plugins/Experimental/PropertyAnimator/Source/PropertyAnimator/Private/Animators/PropertyAnimatorBounce.cpp
--- a/Engine/Plugins/Experimental/PropertyAnimator/Source/PropertyAnimator/Private/Animators/PropertyAnimatorBounce.cpp
+++ b/Engine/Plugins/Experimental/PropertyAnimator/Source/PropertyAnimator/Private/Animators/PropertyAnimatorBounce.cpp
@@ -6,11 +6,6 @@
 #include "Properties/PropertyAnimatorFloatContext.h"
 #include "PropertyAnimatorShared.h"
 
-UPropertyAnimatorBounce::UPropertyAnimatorBounce()
-{
-	SetAnimatorDisplayName(DefaultControllerName);
-}
-
 void UPropertyAnimatorBounce::SetInvertEffect(bool bInvert)
 {
 	if (bInvertEffect == bInvert)
@@ -22,6 +17,13 @@ void UPropertyAnimatorBounce::SetInvertEffect(bool bInvert)
 	OnInvertEffect();
 }
 
+void UPropertyAnimatorBounce::OnAnimatorRegistered(FPropertyAnimatorCoreMetadata& InMetadata)
+{
+	Super::OnAnimatorRegistered(InMetadata);
+
+	InMetadata.Name = TEXT("Bounce");
+}
+
 bool UPropertyAnimatorBounce::EvaluateProperty(const FPropertyAnimatorCoreData& InPropertyData, UPropertyAnimatorCoreContext* InContext, FInstancedPropertyBag& InParameters, FInstancedPropertyBag& OutEvaluationResult) const
 {
 	const float Frequency = InParameters.GetValueDouble(FrequencyParameterName).GetValue();
diff --git a/Engine/Plugins/Experimental/PropertyAnimator/Source/PropertyAnimator/Private/Animators/PropertyAnimatorClock.cpp b/Engine/Plugins/Experimental/PropertyAnimator/Source/PropertyAnimator/Private/Animators/PropertyAnimatorClock.cpp
--- a/Engine/Plugins/Experimental/PropertyAnimator/Source/PropertyAnimator/Private/Animators/PropertyAnimatorClock.cpp
+++ b/Engine/Plugins/Experimental/PropertyAnimator/Source/PropertyAnimator/Private/Animators/PropertyAnimatorClock.cpp
@@ -24,8 +24,6 @@ void UPropertyAnimatorClock::UnregisterFormat(const TCHAR InChar)
 
 UPropertyAnimatorClock::UPropertyAnimatorClock()
 {
-	SetAnimatorDisplayName(DefaultControllerName);
-
 	if (IsTemplate())
 	{
 		RegisterFormat(TEXT('a'), [](const FDateTime& InDateTime)->FString{ return InDateTime.ToFormattedString(TEXT("%a")); });
@@ -87,26 +85,11 @@ FString UPropertyAnimatorClock::FormatDateTime(const FDateTime& InDateTime, cons
     return Result;
 }
 
-EPropertyAnimatorPropertySupport UPropertyAnimatorClock::IsPropertySupported(const FPropertyAnimatorCoreData& InPropertyData) const
+void UPropertyAnimatorClock::OnAnimatorRegistered(FPropertyAnimatorCoreMetadata& InMetadata)
 {
-	if (InPropertyData.IsA<FStrProperty>())
-	{
-		return EPropertyAnimatorPropertySupport::Complete;
-	}
-
-	// Check if a converter supports the conversion
-	if (UPropertyAnimatorCoreSubsystem* AnimatorSubsystem = UPropertyAnimatorCoreSubsystem::Get())
-	{
-		static const FPropertyBagPropertyDesc AnimatorTypeDesc("", EPropertyBagPropertyType::String);
-		const FPropertyBagPropertyDesc PropertyTypeDesc("", InPropertyData.GetLeafProperty());
-
-		if (AnimatorSubsystem->IsConversionSupported(AnimatorTypeDesc, PropertyTypeDesc))
-		{
-			return EPropertyAnimatorPropertySupport::Incomplete;
-		}
-	}
+	Super::OnAnimatorRegistered(InMetadata);
 
-	return Super::IsPropertySupported(InPropertyData);
+	InMetadata.Name = TEXT("Clock");
 }
 
 void UPropertyAnimatorClock::EvaluateProperties(FInstancedPropertyBag& InParameters)
@@ -131,24 +114,4 @@ void UPropertyAnimatorClock::EvaluateProperties(FInstancedPropertyBag& InParamet
 
 		return true;
 	});
-}
-
-void UPropertyAnimatorClock::OnPropertyLinked(UPropertyAnimatorCoreContext* InLinkedProperty, EPropertyAnimatorPropertySupport InSupport)
-{
-	Super::OnPropertyLinked(InLinkedProperty, InSupport);
-
-	if (EnumHasAnyFlags(InSupport, EPropertyAnimatorPropertySupport::Incomplete))
-	{
-		if (const UPropertyAnimatorCoreSubsystem* AnimatorSubsystem = UPropertyAnimatorCoreSubsystem::Get())
-		{
-			static const FPropertyBagPropertyDesc AnimatorTypeDesc("", EPropertyBagPropertyType::String);
-
-			const FPropertyAnimatorCoreData& Property = InLinkedProperty->GetAnimatedProperty();
-			const FPropertyBagPropertyDesc PropertyTypeDesc("", Property.GetLeafProperty());
-			const TSet<UPropertyAnimatorCoreConverterBase*> Converters = AnimatorSubsystem->GetSupportedConverters(AnimatorTypeDesc, PropertyTypeDesc);
-			check(!Converters.IsEmpty())
-
-			InLinkedProperty->SetConverterClass(Converters.Array()[0]->GetClass());
-		}
-	}
 }
\ No newline at end of file
diff --git a/Engine/Plugins/Experimental/PropertyAnimator/Source/PropertyAnimator/Private/Animators/PropertyAnimatorCounter.cpp b/Engine/Plugins/Experimental/PropertyAnimator/Source/PropertyAnimator/Private/Animators/PropertyAnimatorCounter.cpp
--- a/Engine/Plugins/Experimental/PropertyAnimator/Source/PropertyAnimator/Private/Animators/PropertyAnimatorCounter.cpp
+++ b/Engine/Plugins/Experimental/PropertyAnimator/Source/PropertyAnimator/Private/Animators/PropertyAnimatorCounter.cpp
@@ -117,8 +117,6 @@ FString FPropertyAnimatorCounterFormat::FormatNumber(double InNumber) const
 
 UPropertyAnimatorCounter::UPropertyAnimatorCounter()
 {
-	SetAnimatorDisplayName(DefaultAnimatorName);
-
 	if (!IsTemplate())
 	{
 		const TArray<FName> AvailableNames = GetAvailableFormatNames();
@@ -229,6 +227,13 @@ void UPropertyAnimatorCounter::PostEditChangeProperty(FPropertyChangedEvent& InP
 	}
 }
 
+void UPropertyAnimatorCounter::OnAnimatorRegistered(FPropertyAnimatorCoreMetadata& InMetadata)
+{
+	Super::OnAnimatorRegistered(InMetadata);
+
+	InMetadata.Name = TEXT("Counter");
+}
+
 void UPropertyAnimatorCounter::OpenPropertyAnimatorSettings()
 {
 	if (const UPropertyAnimatorSettings* AnimatorSettings = GetDefault<UPropertyAnimatorSettings>())
@@ -254,28 +259,6 @@ void UPropertyAnimatorCounter::SaveCustomFormatAsPreset()
 }
 #endif
 
-EPropertyAnimatorPropertySupport UPropertyAnimatorCounter::IsPropertySupported(const FPropertyAnimatorCoreData& InPropertyData) const
-{
-	if (InPropertyData.IsA<FStrProperty>())
-	{
-		return EPropertyAnimatorPropertySupport::Complete;
-	}
-
-	// Check if a converter supports the conversion
-	if (UPropertyAnimatorCoreSubsystem* AnimatorSubsystem = UPropertyAnimatorCoreSubsystem::Get())
-	{
-		static const FPropertyBagPropertyDesc AnimatorTypeDesc("", EPropertyBagPropertyType::String);
-		const FPropertyBagPropertyDesc PropertyTypeDesc("", InPropertyData.GetLeafProperty());
-
-		if (AnimatorSubsystem->IsConversionSupported(AnimatorTypeDesc, PropertyTypeDesc))
-		{
-			return EPropertyAnimatorPropertySupport::Incomplete;
-		}
-	}
-
-	return Super::IsPropertySupported(InPropertyData);
-}
-
 void UPropertyAnimatorCounter::EvaluateProperties(FInstancedPropertyBag& InParameters)
 {
 	const double TimeElapsed = InParameters.GetValueDouble(TimeElapsedParameterName).GetValue();
@@ -296,26 +279,6 @@ void UPropertyAnimatorCounter::EvaluateProperties(FInstancedPropertyBag& InParam
 	});
 }
 
-void UPropertyAnimatorCounter::OnPropertyLinked(UPropertyAnimatorCoreContext* InLinkedProperty, EPropertyAnimatorPropertySupport InSupport)
-{
-	Super::OnPropertyLinked(InLinkedProperty, InSupport);
-
-	if (EnumHasAnyFlags(InSupport, EPropertyAnimatorPropertySupport::Incomplete))
-	{
-		if (const UPropertyAnimatorCoreSubsystem* AnimatorSubsystem = UPropertyAnimatorCoreSubsystem::Get())
-		{
-			static const FPropertyBagPropertyDesc AnimatorTypeDesc("", EPropertyBagPropertyType::String);
-
-			const FPropertyAnimatorCoreData& Property = InLinkedProperty->GetAnimatedProperty();
-			const FPropertyBagPropertyDesc PropertyTypeDesc("", Property.GetLeafProperty());
-			const TSet<UPropertyAnimatorCoreConverterBase*> Converters = AnimatorSubsystem->GetSupportedConverters(AnimatorTypeDesc, PropertyTypeDesc);
-			check(!Converters.IsEmpty())
-
-			InLinkedProperty->SetConverterClass(Converters.Array()[0]->GetClass());
-		}
-	}
-}
-
 TArray<FName> UPropertyAnimatorCounter::GetAvailableFormatNames() const
 {
 	TArray<FName> FormatNames;
diff --git a/Engine/Plugins/Experimental/PropertyAnimator/Source/PropertyAnimator/Private/Animators/PropertyAnimatorCurve.cpp b/Engine/Plugins/Experimental/PropertyAnimator/Source/PropertyAnimator/Private/Animators/PropertyAnimatorCurve.cpp
--- a/Engine/Plugins/Experimental/PropertyAnimator/Source/PropertyAnimator/Private/Animators/PropertyAnimatorCurve.cpp
+++ b/Engine/Plugins/Experimental/PropertyAnimator/Source/PropertyAnimator/Private/Animators/PropertyAnimatorCurve.cpp
@@ -9,13 +9,11 @@
 
 UPropertyAnimatorCurve::UPropertyAnimatorCurve()
 {
-	SetAnimatorDisplayName(DefaultAnimatorName);
+	static const ConstructorHelpers::FObjectFinder<UPropertyAnimatorWaveCurve> BaseCurve(TEXT("/Script/PropertyAnimator.PropertyAnimatorWaveCurve'/PropertyAnimator/Waves/Constant.Constant'"));
 
-	static const ConstructorHelpers::FObjectFinder<UPropertyAnimatorWaveCurve> SineCurve(TEXT("/Script/PropertyAnimator.PropertyAnimatorWaveCurve'/PropertyAnimator/Waves/Sine.Sine'"));
-
-	if (SineCurve.Succeeded())
+	if (BaseCurve.Succeeded())
 	{
-		WaveCurve = SineCurve.Object;
+		WaveCurve = BaseCurve.Object;
 	}
 
 	static const ConstructorHelpers::FObjectFinder<UPropertyAnimatorEaseCurve> LinearCurve(TEXT("/Script/PropertyAnimator.PropertyAnimatorEaseCurve'/PropertyAnimator/Eases/Linear.Linear'"));
@@ -70,6 +68,13 @@ void UPropertyAnimatorCurve::PostEditChangeProperty(FPropertyChangedEvent& InPro
 }
 #endif
 
+void UPropertyAnimatorCurve::OnAnimatorRegistered(FPropertyAnimatorCoreMetadata& InMetadata)
+{
+	Super::OnAnimatorRegistered(InMetadata);
+
+	InMetadata.Name = TEXT("Curve");
+}
+
 bool UPropertyAnimatorCurve::EvaluateProperty(const FPropertyAnimatorCoreData& InPropertyData, UPropertyAnimatorCoreContext* InContext, FInstancedPropertyBag& InParameters, FInstancedPropertyBag& OutEvaluationResult) const
 {
 	if (!WaveCurve)
diff --git a/Engine/Plugins/Experimental/PropertyAnimator/Source/PropertyAnimator/Private/Animators/PropertyAnimatorNumericBase.cpp b/Engine/Plugins/Experimental/PropertyAnimator/Source/PropertyAnimator/Private/Animators/PropertyAnimatorNumericBase.cpp
--- a/Engine/Plugins/Experimental/PropertyAnimator/Source/PropertyAnimator/Private/Animators/PropertyAnimatorNumericBase.cpp
+++ b/Engine/Plugins/Experimental/PropertyAnimator/Source/PropertyAnimator/Private/Animators/PropertyAnimatorNumericBase.cpp
@@ -276,3 +276,10 @@ bool UPropertyAnimatorNumericBase::IsTimeSourceSupported(UPropertyAnimatorCoreTi
 {
 	return !InTimeSource->IsA<UPropertyAnimatorCoreSystemTimeSource>();
 }
+
+void UPropertyAnimatorNumericBase::OnAnimatorRegistered(FPropertyAnimatorCoreMetadata& InMetadata)
+{
+	Super::OnAnimatorRegistered(InMetadata);
+
+	InMetadata.Category = TEXT("Numeric");
+}
diff --git a/Engine/Plugins/Experimental/PropertyAnimator/Source/PropertyAnimator/Private/Animators/PropertyAnimatorOscillate.cpp b/Engine/Plugins/Experimental/PropertyAnimator/Source/PropertyAnimator/Private/Animators/PropertyAnimatorOscillate.cpp
--- a/Engine/Plugins/Experimental/PropertyAnimator/Source/PropertyAnimator/Private/Animators/PropertyAnimatorOscillate.cpp
+++ b/Engine/Plugins/Experimental/PropertyAnimator/Source/PropertyAnimator/Private/Animators/PropertyAnimatorOscillate.cpp
@@ -5,14 +5,16 @@
 #include "Properties/PropertyAnimatorFloatContext.h"
 #include "PropertyAnimatorShared.h"
 
-UPropertyAnimatorOscillate::UPropertyAnimatorOscillate()
+void UPropertyAnimatorOscillate::SetOscillateFunction(EPropertyAnimatorOscillateFunction InFunction)
 {
-	SetAnimatorDisplayName(DefaultControllerName);
+	OscillateFunction = InFunction;
 }
 
-void UPropertyAnimatorOscillate::SetOscillateFunction(EPropertyAnimatorOscillateFunction InFunction)
+void UPropertyAnimatorOscillate::OnAnimatorRegistered(FPropertyAnimatorCoreMetadata& InMetadata)
 {
-	OscillateFunction = InFunction;
+	Super::OnAnimatorRegistered(InMetadata);
+
+	InMetadata.Name = TEXT("Oscillate");
 }
 
 bool UPropertyAnimatorOscillate::EvaluateProperty(const FPropertyAnimatorCoreData& InPropertyData, UPropertyAnimatorCoreContext* InContext, FInstancedPropertyBag& InParameters, FInstancedPropertyBag& OutEvaluationResult) const
diff --git a/Engine/Plugins/Experimental/PropertyAnimator/Source/PropertyAnimator/Private/Animators/PropertyAnimatorPulse.cpp b/Engine/Plugins/Experimental/PropertyAnimator/Source/PropertyAnimator/Private/Animators/PropertyAnimatorPulse.cpp
--- a/Engine/Plugins/Experimental/PropertyAnimator/Source/PropertyAnimator/Private/Animators/PropertyAnimatorPulse.cpp
+++ b/Engine/Plugins/Experimental/PropertyAnimator/Source/PropertyAnimator/Private/Animators/PropertyAnimatorPulse.cpp
@@ -5,11 +5,6 @@
 #include "Properties/PropertyAnimatorFloatContext.h"
 #include "PropertyAnimatorShared.h"
 
-UPropertyAnimatorPulse::UPropertyAnimatorPulse()
-{
-	SetAnimatorDisplayName(DefaultControllerName);
-}
-
 void UPropertyAnimatorPulse::SetEasingFunction(EPropertyAnimatorEasingFunction InEasingFunction)
 {
 	EasingFunction = InEasingFunction;
@@ -20,6 +15,13 @@ void UPropertyAnimatorPulse::SetEasingType(EPropertyAnimatorEasingType InEasingT
 	EasingType = InEasingType;
 }
 
+void UPropertyAnimatorPulse::OnAnimatorRegistered(FPropertyAnimatorCoreMetadata& InMetadata)
+{
+	Super::OnAnimatorRegistered(InMetadata);
+
+	InMetadata.Name = TEXT("Pulse");
+}
+
 bool UPropertyAnimatorPulse::EvaluateProperty(const FPropertyAnimatorCoreData& InPropertyData, UPropertyAnimatorCoreContext* InContext, FInstancedPropertyBag& InParameters, FInstancedPropertyBag& OutEvaluationResult) const
 {
 	const float Frequency = InParameters.GetValueDouble(FrequencyParameterName).GetValue();
diff --git a/Engine/Plugins/Experimental/PropertyAnimator/Source/PropertyAnimator/Private/Animators/PropertyAnimatorSoundWave.cpp b/Engine/Plugins/Experimental/PropertyAnimator/Source/PropertyAnimator/Private/Animators/PropertyAnimatorSoundWave.cpp
--- a/Engine/Plugins/Experimental/PropertyAnimator/Source/PropertyAnimator/Private/Animators/PropertyAnimatorSoundWave.cpp
+++ b/Engine/Plugins/Experimental/PropertyAnimator/Source/PropertyAnimator/Private/Animators/PropertyAnimatorSoundWave.cpp
@@ -9,7 +9,6 @@
 
 UPropertyAnimatorSoundWave::UPropertyAnimatorSoundWave()
 {
-	SetAnimatorDisplayName(DefaultControllerName);
 	CycleMode = EPropertyAnimatorCycleMode::None;
 }
 
@@ -60,6 +59,13 @@ void UPropertyAnimatorSoundWave::OnSampledSoundWaveChanged()
 #endif
 }
 
+void UPropertyAnimatorSoundWave::OnAnimatorRegistered(FPropertyAnimatorCoreMetadata& InMetadata)
+{
+	Super::OnAnimatorRegistered(InMetadata);
+
+	InMetadata.Name = TEXT("SoundWave");
+}
+
 bool UPropertyAnimatorSoundWave::EvaluateProperty(const FPropertyAnimatorCoreData& InPropertyData, UPropertyAnimatorCoreContext* InContext, FInstancedPropertyBag& InParameters, FInstancedPropertyBag& OutEvaluationResult) const
 {
 	if (AudioAnalyzer && AudioAnalyzer->DurationInSeconds > 0)
diff --git a/Engine/Plugins/Experimental/PropertyAnimator/Source/PropertyAnimator/Private/Animators/PropertyAnimatorTextBase.cpp b/Engine/Plugins/Experimental/PropertyAnimator/Source/PropertyAnimator/Private/Animators/PropertyAnimatorTextBase.cpp
new file mode 100644
--- /dev/null
+++ b/Engine/Plugins/Experimental/PropertyAnimator/Source/PropertyAnimator/Private/Animators/PropertyAnimatorTextBase.cpp
@@ -0,0 +1,55 @@
+// Copyright Epic Games, Inc. All Rights Reserved.
+
+#include "Animators/PropertyAnimatorTextBase.h"
+
+#include "Properties/Converters/PropertyAnimatorCoreConverterBase.h"
+#include "Subsystems/PropertyAnimatorCoreSubsystem.h"
+
+EPropertyAnimatorPropertySupport UPropertyAnimatorTextBase::IsPropertySupported(const FPropertyAnimatorCoreData& InPropertyData) const
+{
+	if (InPropertyData.IsA<FStrProperty>())
+	{
+		return EPropertyAnimatorPropertySupport::Complete;
+	}
+
+	// Check if a converter supports the conversion
+	if (UPropertyAnimatorCoreSubsystem* AnimatorSubsystem = UPropertyAnimatorCoreSubsystem::Get())
+	{
+		static const FPropertyBagPropertyDesc AnimatorTypeDesc("", EPropertyBagPropertyType::String);
+		const FPropertyBagPropertyDesc PropertyTypeDesc("", InPropertyData.GetLeafProperty());
+
+		if (AnimatorSubsystem->IsConversionSupported(AnimatorTypeDesc, PropertyTypeDesc))
+		{
+			return EPropertyAnimatorPropertySupport::Incomplete;
+		}
+	}
+
+	return Super::IsPropertySupported(InPropertyData);
+}
+
+void UPropertyAnimatorTextBase::OnPropertyLinked(UPropertyAnimatorCoreContext* InLinkedProperty, EPropertyAnimatorPropertySupport InSupport)
+{
+	Super::OnPropertyLinked(InLinkedProperty, InSupport);
+
+	if (EnumHasAnyFlags(InSupport, EPropertyAnimatorPropertySupport::Incomplete))
+	{
+		if (const UPropertyAnimatorCoreSubsystem* AnimatorSubsystem = UPropertyAnimatorCoreSubsystem::Get())
+		{
+			static const FPropertyBagPropertyDesc AnimatorTypeDesc("", EPropertyBagPropertyType::String);
+
+			const FPropertyAnimatorCoreData& Property = InLinkedProperty->GetAnimatedProperty();
+			const FPropertyBagPropertyDesc PropertyTypeDesc("", Property.GetLeafProperty());
+			const TSet<UPropertyAnimatorCoreConverterBase*> Converters = AnimatorSubsystem->GetSupportedConverters(AnimatorTypeDesc, PropertyTypeDesc);
+			check(!Converters.IsEmpty())
+
+			InLinkedProperty->SetConverterClass(Converters.Array()[0]->GetClass());
+		}
+	}
+}
+
+void UPropertyAnimatorTextBase::OnAnimatorRegistered(FPropertyAnimatorCoreMetadata& InMetadata)
+{
+	Super::OnAnimatorRegistered(InMetadata);
+
+	InMetadata.Category = TEXT("Text");
+}
diff --git a/Engine/Plugins/Experimental/PropertyAnimator/Source/PropertyAnimator/Private/Animators/PropertyAnimatorTime.cpp b/Engine/Plugins/Experimental/PropertyAnimator/Source/PropertyAnimator/Private/Animators/PropertyAnimatorTime.cpp
--- a/Engine/Plugins/Experimental/PropertyAnimator/Source/PropertyAnimator/Private/Animators/PropertyAnimatorTime.cpp
+++ b/Engine/Plugins/Experimental/PropertyAnimator/Source/PropertyAnimator/Private/Animators/PropertyAnimatorTime.cpp
@@ -5,9 +5,11 @@
 #include "Properties/Handlers/PropertyAnimatorCoreHandlerBase.h"
 #include "Properties/PropertyAnimatorFloatContext.h"
 
-UPropertyAnimatorTime::UPropertyAnimatorTime()
+void UPropertyAnimatorTime::OnAnimatorRegistered(FPropertyAnimatorCoreMetadata& InMetadata)
 {
-	SetAnimatorDisplayName(DefaultControllerName);
+	Super::OnAnimatorRegistered(InMetadata);
+
+	InMetadata.Name = TEXT("Time");
 }
 
 bool UPropertyAnimatorTime::EvaluateProperty(const FPropertyAnimatorCoreData& InPropertyData, UPropertyAnimatorCoreContext* InContext, FInstancedPropertyBag& InParameters, FInstancedPropertyBag& OutEvaluationResult) const
diff --git a/Engine/Plugins/Experimental/PropertyAnimator/Source/PropertyAnimator/Private/Animators/PropertyAnimatorWiggle.cpp b/Engine/Plugins/Experimental/PropertyAnimator/Source/PropertyAnimator/Private/Animators/PropertyAnimatorWiggle.cpp
--- a/Engine/Plugins/Experimental/PropertyAnimator/Source/PropertyAnimator/Private/Animators/PropertyAnimatorWiggle.cpp
+++ b/Engine/Plugins/Experimental/PropertyAnimator/Source/PropertyAnimator/Private/Animators/PropertyAnimatorWiggle.cpp
@@ -10,13 +10,18 @@ UPropertyAnimatorWiggle::UPropertyAnimatorWiggle()
 {
 	static int32 SeedIncrement = 0;
 
-	SetAnimatorDisplayName(DefaultControllerName);
-
 	bRandomTimeOffset = true;
 	Seed = SeedIncrement++;
 	CycleMode = EPropertyAnimatorCycleMode::None;
 }
 
+void UPropertyAnimatorWiggle::OnAnimatorRegistered(FPropertyAnimatorCoreMetadata& InMetadata)
+{
+	Super::OnAnimatorRegistered(InMetadata);
+
+	InMetadata.Name = TEXT("Wiggle");
+}
+
 bool UPropertyAnimatorWiggle::EvaluateProperty(const FPropertyAnimatorCoreData& InPropertyData, UPropertyAnimatorCoreContext* InContext, FInstancedPropertyBag& InParameters, FInstancedPropertyBag& OutEvaluationResult) const
 {
 	const double TimeElapsed = InParameters.GetValueDouble(TimeElapsedParameterName).GetValue();
diff --git a/Engine/Plugins/Experimental/PropertyAnimator/Source/PropertyAnimator/Public/Animators/PropertyAnimatorBounce.h b/Engine/Plugins/Experimental/PropertyAnimator/Source/PropertyAnimator/Public/Animators/PropertyAnimatorBounce.h
--- a/Engine/Plugins/Experimental/PropertyAnimator/Source/PropertyAnimator/Public/Animators/PropertyAnimatorBounce.h
+++ b/Engine/Plugins/Experimental/PropertyAnimator/Source/PropertyAnimator/Public/Animators/PropertyAnimatorBounce.h
@@ -14,10 +14,6 @@ class UPropertyAnimatorBounce : public UPropertyAnimatorNumericBase
 	GENERATED_BODY()
 
 public:
-	static constexpr const TCHAR* DefaultControllerName = TEXT("Bounce");
-
-	UPropertyAnimatorBounce();
-
 	PROPERTYANIMATOR_API void SetInvertEffect(bool bInvert);
 	bool GetInvertEffect() const
 	{
@@ -28,6 +24,7 @@ protected:
 	virtual void OnInvertEffect() {}
 
 	//~ Begin UPropertyAnimatorFloatBase
+	virtual void OnAnimatorRegistered(FPropertyAnimatorCoreMetadata& InMetadata) override;
 	virtual bool EvaluateProperty(const FPropertyAnimatorCoreData& InPropertyData, UPropertyAnimatorCoreContext* InContext, FInstancedPropertyBag& InParameters, FInstancedPropertyBag& OutEvaluationResult) const override;
 	//~ End UPropertyAnimatorFloatBase
 
diff --git a/Engine/Plugins/Experimental/PropertyAnimator/Source/PropertyAnimator/Public/Animators/PropertyAnimatorClock.h b/Engine/Plugins/Experimental/PropertyAnimator/Source/PropertyAnimator/Public/Animators/PropertyAnimatorClock.h
--- a/Engine/Plugins/Experimental/PropertyAnimator/Source/PropertyAnimator/Public/Animators/PropertyAnimatorClock.h
+++ b/Engine/Plugins/Experimental/PropertyAnimator/Source/PropertyAnimator/Public/Animators/PropertyAnimatorClock.h
@@ -2,9 +2,9 @@
 
 #pragma once
 
-#include "Animators/PropertyAnimatorCoreBase.h"
 #include "HAL/Platform.h"
 #include "Misc/Timespan.h"
+#include "PropertyAnimatorTextBase.h"
 #include "PropertyAnimatorClock.generated.h"
 
 /** Mode supported for properties value */
@@ -21,13 +21,11 @@ enum class EPropertyAnimatorClockMode : uint8
 
 /** Animate supported string properties to display time */
 UCLASS(MinimalAPI, AutoExpandCategories=("Animator"))
-class UPropertyAnimatorClock : public UPropertyAnimatorCoreBase
+class UPropertyAnimatorClock : public UPropertyAnimatorTextBase
 {
 	GENERATED_BODY()
 
 public:
-	static constexpr const TCHAR* DefaultControllerName = TEXT("Clock");
-
 	static void RegisterFormat(TCHAR InChar, TFunction<FString(const FDateTime&)> InFormatter);
 	static void UnregisterFormat(TCHAR InChar);
 	static FString FormatDateTime(const FDateTime& InDateTime, const FString& InDisplayFormat);
@@ -42,9 +40,8 @@ public:
 
 protected:
 	//~ Begin UPropertyAnimatorCoreBase
-	virtual EPropertyAnimatorPropertySupport IsPropertySupported(const FPropertyAnimatorCoreData& InPropertyData) const override;
+	virtual void OnAnimatorRegistered(FPropertyAnimatorCoreMetadata& InMetadata) override;
 	virtual void EvaluateProperties(FInstancedPropertyBag& InParameters) override;
-	virtual void OnPropertyLinked(UPropertyAnimatorCoreContext* InLinkedProperty, EPropertyAnimatorPropertySupport InSupport) override;
 	//~ End UPropertyAnimatorCoreBase
 
 	/**
diff --git a/Engine/Plugins/Experimental/PropertyAnimator/Source/PropertyAnimator/Public/Animators/PropertyAnimatorCounter.h b/Engine/Plugins/Experimental/PropertyAnimator/Source/PropertyAnimator/Public/Animators/PropertyAnimatorCounter.h
--- a/Engine/Plugins/Experimental/PropertyAnimator/Source/PropertyAnimator/Public/Animators/PropertyAnimatorCounter.h
+++ b/Engine/Plugins/Experimental/PropertyAnimator/Source/PropertyAnimator/Public/Animators/PropertyAnimatorCounter.h
@@ -2,7 +2,7 @@
 
 #pragma once
 
-#include "Animators/PropertyAnimatorNumericBase.h"
+#include "PropertyAnimatorTextBase.h"
 #include "StructUtils/InstancedStruct.h"
 #include "PropertyAnimatorCounter.generated.h"
 
@@ -92,13 +92,11 @@ struct FPropertyAnimatorCounterFormat
 
 /** Animate supported string properties to display a counter */
 UCLASS(MinimalAPI, AutoExpandCategories=("Animator"))
-class UPropertyAnimatorCounter : public UPropertyAnimatorCoreBase
+class UPropertyAnimatorCounter : public UPropertyAnimatorTextBase
 {
 	GENERATED_BODY()
 
 public:
-	static constexpr const TCHAR* DefaultAnimatorName = TEXT("Counter");
-
 #if WITH_EDITOR
 	PROPERTYANIMATOR_API static FName GetUseCustomFormatPropertyName();
 #endif
@@ -136,9 +134,8 @@ protected:
 	//~ End UObject
 
 	//~ Begin UPropertyAnimatorCoreBase
-	virtual EPropertyAnimatorPropertySupport IsPropertySupported(const FPropertyAnimatorCoreData& InPropertyData) const override;
+	virtual void OnAnimatorRegistered(FPropertyAnimatorCoreMetadata& InMetadata) override;
 	virtual void EvaluateProperties(FInstancedPropertyBag& InParameters) override;
-	virtual void OnPropertyLinked(UPropertyAnimatorCoreContext* InLinkedProperty, EPropertyAnimatorPropertySupport InSupport) override;
 	//~ End UPropertyAnimatorCoreBase
 
 	UFUNCTION()
diff --git a/Engine/Plugins/Experimental/PropertyAnimator/Source/PropertyAnimator/Public/Animators/PropertyAnimatorCurve.h b/Engine/Plugins/Experimental/PropertyAnimator/Source/PropertyAnimator/Public/Animators/PropertyAnimatorCurve.h
--- a/Engine/Plugins/Experimental/PropertyAnimator/Source/PropertyAnimator/Public/Animators/PropertyAnimatorCurve.h
+++ b/Engine/Plugins/Experimental/PropertyAnimator/Source/PropertyAnimator/Public/Animators/PropertyAnimatorCurve.h
@@ -13,10 +13,10 @@ struct FPropertyAnimatorCurveEasing
 {
 	GENERATED_BODY()
 
-	UPROPERTY(EditInstanceOnly, Category="Animator")
+	UPROPERTY(EditInstanceOnly, DisplayName="Curve", Category="Animator")
 	TObjectPtr<UPropertyAnimatorEaseCurve> EaseCurve;
 
-	UPROPERTY(EditInstanceOnly, Interp, Category="Animator", meta=(ClampMin="0", Units=Seconds))
+	UPROPERTY(EditInstanceOnly, Interp, DisplayName="Duration", Category="Animator", meta=(ClampMin="0", Units=Seconds))
 	float EaseDuration = 0.f;
 };
 
@@ -29,8 +29,6 @@ class UPropertyAnimatorCurve : public UPropertyAnimatorNumericBase
 	GENERATED_BODY()
 
 public:
-	static constexpr const TCHAR* DefaultAnimatorName = TEXT("Curve");
-
 	UPropertyAnimatorCurve();
 
 	void SetWaveCurve(UPropertyAnimatorWaveCurve* InCurve);
@@ -76,6 +74,7 @@ protected:
 	//~ End UObject
 
 	//~ Begin UPropertyAnimatorFloatBase
+	virtual void OnAnimatorRegistered(FPropertyAnimatorCoreMetadata& InMetadata) override;
 	virtual bool EvaluateProperty(const FPropertyAnimatorCoreData& InPropertyData, UPropertyAnimatorCoreContext* InContext, FInstancedPropertyBag& InParameters, FInstancedPropertyBag& OutEvaluationResult) const override;
 	//~ End UPropertyAnimatorFloatBase
 
@@ -83,10 +82,6 @@ protected:
 	void OnEaseOutChanged();
 	virtual void OnCycleDurationChanged() override;
 
-	/** The wave curve to sample for the animation, base curve */
-	UPROPERTY(EditInstanceOnly, Setter, Getter, DisplayName="Base Curve", Category="Animator")
-	TObjectPtr<UPropertyAnimatorWaveCurve> WaveCurve;
-
 	/** Use ease in effect */
 	UPROPERTY(EditInstanceOnly, Setter="SetEaseInEnabled", Getter="GetEaseInEnabled", DisplayName="InEnabled", Category="Animator", meta=(InlineEditConditionToggle))
 	bool bEaseInEnabled = false;
@@ -95,6 +90,10 @@ protected:
 	UPROPERTY(EditInstanceOnly, Setter, Getter, DisplayName="In", Category="Animator", meta=(EditCondition="bEaseInEnabled"))
 	FPropertyAnimatorCurveEasing EaseIn;
 
+	/** The base curve to sample for the animation */
+	UPROPERTY(EditInstanceOnly, Setter, Getter, DisplayName="Loop Curve", Category="Animator")
+	TObjectPtr<UPropertyAnimatorWaveCurve> WaveCurve;
+
 	/** Use ease out effect */
 	UPROPERTY(EditInstanceOnly, Setter="SetEaseOutEnabled", Getter="GetEaseOutEnabled", DisplayName="OutEnabled", Category="Animator", meta=(InlineEditConditionToggle))
 	bool bEaseOutEnabled = false;
diff --git a/Engine/Plugins/Experimental/PropertyAnimator/Source/PropertyAnimator/Public/Animators/PropertyAnimatorNumericBase.h b/Engine/Plugins/Experimental/PropertyAnimator/Source/PropertyAnimator/Public/Animators/PropertyAnimatorNumericBase.h
--- a/Engine/Plugins/Experimental/PropertyAnimator/Source/PropertyAnimator/Public/Animators/PropertyAnimatorNumericBase.h
+++ b/Engine/Plugins/Experimental/PropertyAnimator/Source/PropertyAnimator/Public/Animators/PropertyAnimatorNumericBase.h
@@ -8,7 +8,7 @@
 UENUM(BlueprintType)
 enum class EPropertyAnimatorCycleMode : uint8
 {
-	/** Don't cycle time at all */
+	/** Disable cycle options */
 	None	UMETA(Hidden),
 	/** Cycle only once then stop */
 	DoOnce,
@@ -96,6 +96,7 @@ protected:
 	virtual void EvaluateProperties(FInstancedPropertyBag& InParameters) override;
 	virtual void OnPropertyLinked(UPropertyAnimatorCoreContext* InLinkedProperty, EPropertyAnimatorPropertySupport InSupport) override;
 	virtual bool IsTimeSourceSupported(UPropertyAnimatorCoreTimeSourceBase* InTimeSource) const override;
+	virtual void OnAnimatorRegistered(FPropertyAnimatorCoreMetadata& InMetadata) override;
 	//~ End UPropertyAnimatorCoreBase
 
 	/** Evaluate and return float value for a property */
diff --git a/Engine/Plugins/Experimental/PropertyAnimator/Source/PropertyAnimator/Public/Animators/PropertyAnimatorOscillate.h b/Engine/Plugins/Experimental/PropertyAnimator/Source/PropertyAnimator/Public/Animators/PropertyAnimatorOscillate.h
--- a/Engine/Plugins/Experimental/PropertyAnimator/Source/PropertyAnimator/Public/Animators/PropertyAnimatorOscillate.h
+++ b/Engine/Plugins/Experimental/PropertyAnimator/Source/PropertyAnimator/Public/Animators/PropertyAnimatorOscillate.h
@@ -25,10 +25,6 @@ class UPropertyAnimatorOscillate : public UPropertyAnimatorNumericBase
 	GENERATED_BODY()
 
 public:
-	static constexpr const TCHAR* DefaultControllerName = TEXT("Oscillate");
-
-	UPropertyAnimatorOscillate();
-
 	PROPERTYANIMATOR_API void SetOscillateFunction(EPropertyAnimatorOscillateFunction InFunction);
 	EPropertyAnimatorOscillateFunction GetOscillateFunction() const
 	{
@@ -37,6 +33,7 @@ public:
 
 protected:
 	//~ Begin UPropertyAnimatorFloatBase
+	virtual void OnAnimatorRegistered(FPropertyAnimatorCoreMetadata& InMetadata) override;
 	virtual bool EvaluateProperty(const FPropertyAnimatorCoreData& InPropertyData, UPropertyAnimatorCoreContext* InContext, FInstancedPropertyBag& InParameters, FInstancedPropertyBag& OutEvaluationResult) const override;
 	//~ End UPropertyAnimatorFloatBase
 
diff --git a/Engine/Plugins/Experimental/PropertyAnimator/Source/PropertyAnimator/Public/Animators/PropertyAnimatorPulse.h b/Engine/Plugins/Experimental/PropertyAnimator/Source/PropertyAnimator/Public/Animators/PropertyAnimatorPulse.h
--- a/Engine/Plugins/Experimental/PropertyAnimator/Source/PropertyAnimator/Public/Animators/PropertyAnimatorPulse.h
+++ b/Engine/Plugins/Experimental/PropertyAnimator/Source/PropertyAnimator/Public/Animators/PropertyAnimatorPulse.h
@@ -17,10 +17,6 @@ class UPropertyAnimatorPulse : public UPropertyAnimatorNumericBase
 	GENERATED_BODY()
 
 public:
-	static constexpr const TCHAR* DefaultControllerName = TEXT("Pulse");
-
-	UPropertyAnimatorPulse();
-
 	PROPERTYANIMATOR_API void SetEasingFunction(EPropertyAnimatorEasingFunction InEasingFunction);
 	EPropertyAnimatorEasingFunction GetEasingFunction() const
 	{
@@ -35,6 +31,7 @@ public:
 
 protected:
 	//~ Begin UPropertyAnimatorFloatBase
+	virtual void OnAnimatorRegistered(FPropertyAnimatorCoreMetadata& InMetadata) override;
 	virtual bool EvaluateProperty(const FPropertyAnimatorCoreData& InPropertyData, UPropertyAnimatorCoreContext* InContext, FInstancedPropertyBag& InParameters, FInstancedPropertyBag& OutEvaluationResult) const override;
 	//~ End UPropertyAnimatorFloatBase
 
diff --git a/Engine/Plugins/Experimental/PropertyAnimator/Source/PropertyAnimator/Public/Animators/PropertyAnimatorSoundWave.h b/Engine/Plugins/Experimental/PropertyAnimator/Source/PropertyAnimator/Public/Animators/PropertyAnimatorSoundWave.h
--- a/Engine/Plugins/Experimental/PropertyAnimator/Source/PropertyAnimator/Public/Animators/PropertyAnimatorSoundWave.h
+++ b/Engine/Plugins/Experimental/PropertyAnimator/Source/PropertyAnimator/Public/Animators/PropertyAnimatorSoundWave.h
@@ -17,8 +17,6 @@ class UPropertyAnimatorSoundWave : public UPropertyAnimatorNumericBase
 	GENERATED_BODY()
 
 public:
-	static constexpr const TCHAR* DefaultControllerName = TEXT("SoundWave");
-
 	UPropertyAnimatorSoundWave();
 
 	PROPERTYANIMATOR_API void SetSampledSoundWave(USoundWave* InSoundWave);
@@ -43,6 +41,7 @@ protected:
 	void OnSampledSoundWaveChanged();
 
 	//~ Begin UPropertyAnimatorFloatBase
+	virtual void OnAnimatorRegistered(FPropertyAnimatorCoreMetadata& InMetadata) override;
 	virtual bool EvaluateProperty(const FPropertyAnimatorCoreData& InPropertyData, UPropertyAnimatorCoreContext* InContext, FInstancedPropertyBag& InParameters, FInstancedPropertyBag& OutEvaluationResult) const override;
 	//~ End UPropertyAnimatorFloatBase
 
diff --git a/Engine/Plugins/Experimental/PropertyAnimator/Source/PropertyAnimator/Public/Animators/PropertyAnimatorTextBase.h b/Engine/Plugins/Experimental/PropertyAnimator/Source/PropertyAnimator/Public/Animators/PropertyAnimatorTextBase.h
new file mode 100644
--- /dev/null
+++ b/Engine/Plugins/Experimental/PropertyAnimator/Source/PropertyAnimator/Public/Animators/PropertyAnimatorTextBase.h
@@ -0,0 +1,22 @@
+// Copyright Epic Games, Inc. All Rights Reserved.
+
+#pragma once
+
+#include "Animators/PropertyAnimatorCoreBase.h"
+#include "PropertyAnimatorTextBase.generated.h"
+
+/**
+ * Animate supported text properties with various options
+ */
+UCLASS(MinimalAPI, Abstract, AutoExpandCategories=("Animator"))
+class UPropertyAnimatorTextBase : public UPropertyAnimatorCoreBase
+{
+	GENERATED_BODY()
+
+protected:
+	//~ Begin UPropertyAnimatorCoreBase
+	virtual EPropertyAnimatorPropertySupport IsPropertySupported(const FPropertyAnimatorCoreData& InPropertyData) const override;
+	virtual void OnPropertyLinked(UPropertyAnimatorCoreContext* InLinkedProperty, EPropertyAnimatorPropertySupport InSupport) override;
+	virtual void OnAnimatorRegistered(FPropertyAnimatorCoreMetadata& InMetadata) override;
+	//~ End UPropertyAnimatorCoreBase
+};
\ No newline at end of file
diff --git a/Engine/Plugins/Experimental/PropertyAnimator/Source/PropertyAnimator/Public/Animators/PropertyAnimatorTime.h b/Engine/Plugins/Experimental/PropertyAnimator/Source/PropertyAnimator/Public/Animators/PropertyAnimatorTime.h
--- a/Engine/Plugins/Experimental/PropertyAnimator/Source/PropertyAnimator/Public/Animators/PropertyAnimatorTime.h
+++ b/Engine/Plugins/Experimental/PropertyAnimator/Source/PropertyAnimator/Public/Animators/PropertyAnimatorTime.h
@@ -13,13 +13,9 @@ class UPropertyAnimatorTime : public UPropertyAnimatorNumericBase
 {
 	GENERATED_BODY()
 
-public:
-	static constexpr const TCHAR* DefaultControllerName = TEXT("Time");
-
-	UPropertyAnimatorTime();
-
 protected:
 	//~ Begin UPropertyAnimatorFloatBase
+	virtual void OnAnimatorRegistered(FPropertyAnimatorCoreMetadata& InMetadata) override;
 	virtual bool EvaluateProperty(const FPropertyAnimatorCoreData& InPropertyData, UPropertyAnimatorCoreContext* InContext, FInstancedPropertyBag& InParameters, FInstancedPropertyBag& OutEvaluationResult) const override;
 	//~ End UPropertyAnimatorFloatBase
 };
\ No newline at end of file
diff --git a/Engine/Plugins/Experimental/PropertyAnimator/Source/PropertyAnimator/Public/Animators/PropertyAnimatorWiggle.h b/Engine/Plugins/Experimental/PropertyAnimator/Source/PropertyAnimator/Public/Animators/PropertyAnimatorWiggle.h
--- a/Engine/Plugins/Experimental/PropertyAnimator/Source/PropertyAnimator/Public/Animators/PropertyAnimatorWiggle.h
+++ b/Engine/Plugins/Experimental/PropertyAnimator/Source/PropertyAnimator/Public/Animators/PropertyAnimatorWiggle.h
@@ -14,12 +14,11 @@ class UPropertyAnimatorWiggle : public UPropertyAnimatorNumericBase
 	GENERATED_BODY()
 
 public:
-	static constexpr const TCHAR* DefaultControllerName = TEXT("Wiggle");
-
 	UPropertyAnimatorWiggle();
 
 protected:
 	//~ Begin UPropertyAnimatorFloatBase
+	virtual void OnAnimatorRegistered(FPropertyAnimatorCoreMetadata& InMetadata) override;
 	virtual bool EvaluateProperty(const FPropertyAnimatorCoreData& InPropertyData, UPropertyAnimatorCoreContext* InContext, FInstancedPropertyBag& InParameters, FInstancedPropertyBag& OutEvaluationResult) const override;
 	//~ End UPropertyAnimatorFloatBase
 };
\ No newline at end of file
diff --git a/Engine/Plugins/Experimental/PropertyAnimatorCore/Source/PropertyAnimatorCore/Private/Animators/PropertyAnimatorCoreBase.cpp b/Engine/Plugins/Experimental/PropertyAnimatorCore/Source/PropertyAnimatorCore/Private/Animators/PropertyAnimatorCoreBase.cpp
--- a/Engine/Plugins/Experimental/PropertyAnimatorCore/Source/PropertyAnimatorCore/Private/Animators/PropertyAnimatorCoreBase.cpp
+++ b/Engine/Plugins/Experimental/PropertyAnimatorCore/Source/PropertyAnimatorCore/Private/Animators/PropertyAnimatorCoreBase.cpp
@@ -33,12 +33,14 @@ FName UPropertyAnimatorCoreBase::GetLinkedPropertiesPropertyName()
 
 UPropertyAnimatorCoreBase::UPropertyAnimatorCoreBase()
 {
-#if WITH_EDITOR
 	if (!IsTemplate())
 	{
+		Metadata = GetClass()->GetDefaultObject<UPropertyAnimatorCoreBase>()->Metadata;
+
+#if WITH_EDITOR
 		FCoreUObjectDelegates::OnObjectsReplaced.AddUObject(this, &UPropertyAnimatorCoreBase::OnObjectReplaced);
-	}
 #endif
+	}
 }
 
 UPropertyAnimatorCoreComponent* UPropertyAnimatorCoreBase::GetAnimatorComponent() const
@@ -104,6 +106,18 @@ UPropertyAnimatorCoreContext* UPropertyAnimatorCoreBase::GetLinkedPropertyContex
 	return PropertyOptions ? *PropertyOptions : nullptr;
 }
 
+void UPropertyAnimatorCoreBase::PostCDOContruct()
+{
+	Super::PostCDOContruct();
+
+	if (IsTemplate() && !Metadata.IsValid())
+	{
+		Metadata = MakeShared<FPropertyAnimatorCoreMetadata>();
+		OnAnimatorRegistered(*Metadata);
+		SetAnimatorDisplayName(Metadata->Name);
+	}
+}
+
 void UPropertyAnimatorCoreBase::BeginDestroy()
 {
 	Super::BeginDestroy();
@@ -241,10 +255,24 @@ void UPropertyAnimatorCoreBase::SetTimeSourceName(FName InTimeSourceName)
 	OnTimeSourceNameChanged();
 }
 
+FName UPropertyAnimatorCoreBase::GetAnimatorCategory() const
+{
+	if (Metadata.IsValid())
+	{
+		return Metadata->Category;
+	}
+
+	return NAME_None;
+}
+
 FName UPropertyAnimatorCoreBase::GetAnimatorOriginalName() const
 {
-	const UPropertyAnimatorCoreBase* CDO = GetClass()->GetDefaultObject<UPropertyAnimatorCoreBase>();
-	return CDO ? CDO->AnimatorDisplayName : NAME_None;
+	if (Metadata.IsValid())
+	{
+		return Metadata->Name;
+	}
+
+	return NAME_None;
 }
 
 bool UPropertyAnimatorCoreBase::GetPropertiesSupported(const FPropertyAnimatorCoreData& InPropertyData, TSet<FPropertyAnimatorCoreData>& OutProperties, uint8 InSearchDepth, EPropertyAnimatorPropertySupport InSupportExpected) const
@@ -348,7 +376,7 @@ void UPropertyAnimatorCoreBase::OnAnimatorEnabled()
 		, Log
 		, TEXT("%s : PropertyAnimator %s (%s) enabled")
 		, GetAnimatorActor() ? *GetAnimatorActor()->GetActorNameOrLabel() : TEXT("Invalid Actor")
-		, *GetAnimatorDisplayName()
+		, *GetAnimatorDisplayName().ToString()
 		, *GetAnimatorOriginalName().ToString());
 }
 
@@ -358,7 +386,7 @@ void UPropertyAnimatorCoreBase::OnAnimatorDisabled()
 		, Log
 		, TEXT("%s : PropertyAnimator %s (%s) disabled")
 		, GetAnimatorActor() ? *GetAnimatorActor()->GetActorNameOrLabel() : TEXT("Invalid Actor")
-		, *GetAnimatorDisplayName()
+		, *GetAnimatorDisplayName().ToString()
 		, *GetAnimatorOriginalName().ToString());
 
 	constexpr bool bForceReset = true;
diff --git a/Engine/Plugins/Experimental/PropertyAnimatorCore/Source/PropertyAnimatorCore/Public/Animators/PropertyAnimatorCoreBase.h b/Engine/Plugins/Experimental/PropertyAnimatorCore/Source/PropertyAnimatorCore/Public/Animators/PropertyAnimatorCoreBase.h
--- a/Engine/Plugins/Experimental/PropertyAnimatorCore/Source/PropertyAnimatorCore/Public/Animators/PropertyAnimatorCoreBase.h
+++ b/Engine/Plugins/Experimental/PropertyAnimatorCore/Source/PropertyAnimatorCore/Public/Animators/PropertyAnimatorCoreBase.h
@@ -20,6 +20,12 @@ enum class EPropertyAnimatorPropertySupport : uint8
 	All = Incomplete | Complete
 };
 
+struct FPropertyAnimatorCoreMetadata
+{
+	FName Name;
+	FName Category = TEXT("Default");
+};
+
 /** Abstract base class for any Animator, holds a set of linked properties */
 UCLASS(MinimalAPI, Abstract, EditInlineNew, AutoExpandCategories=("Animator"))
 class UPropertyAnimatorCoreBase : public UObject
@@ -99,11 +105,14 @@ public:
 
 	/** Set the display name of this animator */
 	PROPERTYANIMATORCORE_API void SetAnimatorDisplayName(FName InName);
-	FString GetAnimatorDisplayName() const
+
+	FName GetAnimatorDisplayName() const
 	{
-		return AnimatorDisplayName.ToString();
+		return AnimatorDisplayName;
 	}
 
+	PROPERTYANIMATORCORE_API FName GetAnimatorCategory() const;
+
 	/** Gets the Animator original name */
 	PROPERTYANIMATORCORE_API FName GetAnimatorOriginalName() const;
 
@@ -202,6 +211,7 @@ public:
 
 protected:
 	//~ Begin UObject
+	PROPERTYANIMATORCORE_API virtual void PostCDOContruct() override;
 	PROPERTYANIMATORCORE_API virtual void BeginDestroy() override;
 	PROPERTYANIMATORCORE_API virtual void PostLoad() override;
 	PROPERTYANIMATORCORE_API virtual void PostEditImport() override;
@@ -265,6 +275,8 @@ protected:
 	PROPERTYANIMATORCORE_API virtual void OnAnimatorEnabled();
 	PROPERTYANIMATORCORE_API virtual void OnAnimatorDisabled();
 
+	virtual void OnAnimatorRegistered(FPropertyAnimatorCoreMetadata& InMetadata) {}
+
 	virtual void OnTimeSourceChanged() {}
 
 	/** Returns the property context class to use */
@@ -362,4 +374,7 @@ private:
 
 	/** Are we evaluating properties currently */
 	bool bEvaluatingProperties = false;
+
+	/** Animator metadata, same for all instances of this class */
+	TSharedPtr<FPropertyAnimatorCoreMetadata> Metadata;
 };
\ No newline at end of file
diff --git a/Engine/Plugins/Experimental/PropertyAnimatorCore/Source/PropertyAnimatorCoreEditor/Private/Customizations/PropertyAnimatorCoreEditorStackCustomization.cpp b/Engine/Plugins/Experimental/PropertyAnimatorCore/Source/PropertyAnimatorCoreEditor/Private/Customizations/PropertyAnimatorCoreEditorStackCustomization.cpp
--- a/Engine/Plugins/Experimental/PropertyAnimatorCore/Source/PropertyAnimatorCoreEditor/Private/Customizations/PropertyAnimatorCoreEditorStackCustomization.cpp
+++ b/Engine/Plugins/Experimental/PropertyAnimatorCore/Source/PropertyAnimatorCoreEditor/Private/Customizations/PropertyAnimatorCoreEditorStackCustomization.cpp
@@ -215,10 +215,12 @@ void UPropertyAnimatorCoreEditorStackCustomization::CustomizeItemHeader(const FO
 		// Commands for item on key events
 		const TSharedPtr<FUICommandList> AnimatorCommands = CreateAnimatorCommands(InItem);
 
+		const FString AnimatorDisplayName = Animator->GetAnimatorDisplayName().ToString();
+
 		const TSet<FString> SearchKeywords
 		{
 			Animator->GetAnimatorOriginalName().ToString(),
-			Animator->GetAnimatorDisplayName()
+			AnimatorDisplayName
 		};
 
 		/** Show last execution error messages if failed execution */
@@ -228,7 +230,7 @@ void UPropertyAnimatorCoreEditorStackCustomization::CustomizeItemHeader(const FO
 
 		if (InItem->GetValueCount() == 1)
 		{
-			HeaderLabel = Animator->GetAnimatorDisplayName();
+			HeaderLabel = AnimatorDisplayName;
 
 			TWeakObjectPtr<UPropertyAnimatorCoreBase> AnimatorWeak(Animator);
 
diff --git a/Engine/Plugins/Experimental/PropertyAnimatorCore/Source/PropertyAnimatorCoreEditor/Private/Customizations/PropertyAnimatorCoreEditorStackCustomization.h b/Engine/Plugins/Experimental/PropertyAnimatorCore/Source/PropertyAnimatorCoreEditor/Private/Customizations/PropertyAnimatorCoreEditorStackCustomization.h
--- a/Engine/Plugins/Experimental/PropertyAnimatorCore/Source/PropertyAnimatorCoreEditor/Private/Customizations/PropertyAnimatorCoreEditorStackCustomization.h
+++ b/Engine/Plugins/Experimental/PropertyAnimatorCore/Source/PropertyAnimatorCoreEditor/Private/Customizations/PropertyAnimatorCoreEditorStackCustomization.h
@@ -2,7 +2,6 @@
 
 #pragma once
 
-#include "Animators/PropertyAnimatorCoreBase.h"
 #include "Customizations/OperatorStackEditorStackCustomization.h"
 #include "PropertyAnimatorCoreEditorStackCustomization.generated.h"
 
diff --git a/Engine/Plugins/Experimental/PropertyAnimatorCore/Source/PropertyAnimatorCoreEditor/Private/Menus/PropertyAnimatorCoreEditorMenu.cpp b/Engine/Plugins/Experimental/PropertyAnimatorCore/Source/PropertyAnimatorCoreEditor/Private/Menus/PropertyAnimatorCoreEditorMenu.cpp
--- a/Engine/Plugins/Experimental/PropertyAnimatorCore/Source/PropertyAnimatorCoreEditor/Private/Menus/PropertyAnimatorCoreEditorMenu.cpp
+++ b/Engine/Plugins/Experimental/PropertyAnimatorCore/Source/PropertyAnimatorCoreEditor/Private/Menus/PropertyAnimatorCoreEditorMenu.cpp
@@ -43,6 +43,7 @@ void UE::PropertyAnimatorCoreEditor::Menu::FillNewAnimatorSection(UToolMenu* InM
 	for (UPropertyAnimatorCoreBase* NewAnimator : NewAvailableAnimators)
 	{
 		const FName MenuName = NewAnimator->GetAnimatorOriginalName();
+		const FName MenuCategory = NewAnimator->GetAnimatorCategory();
 		const FText MenuLabel = FText::FromName(MenuName);
 		const FText MenuTooltip = LOCTEXT("NewAnimator.Tooltip", "Create a new animator");
 		const FSlateIcon MenuIcon = FSlateIconFinder::FindIconForClass(NewAnimator->GetClass());
@@ -63,7 +64,9 @@ void UE::PropertyAnimatorCoreEditor::Menu::FillNewAnimatorSection(UToolMenu* InM
 		{
 			UPropertyAnimatorCorePresetBase* EmptyPreset = nullptr;
 
-			NewAnimatorsSection.AddMenuEntry(
+			FToolMenuSection& AnimatorCategorySection = InMenu->FindOrAddSection(MenuCategory, FText::FromName(MenuCategory));
+
+			AnimatorCategorySection.AddMenuEntry(
 				MenuName
 				, MenuLabel
 				, MenuTooltip
@@ -96,7 +99,7 @@ void UE::PropertyAnimatorCoreEditor::Menu::FillExistingAnimatorSection(UToolMenu
 	{
 		for (UPropertyAnimatorCoreBase* Animator : Subsystem->GetExistingAnimators(Property))
 		{
-			const FName MenuName(Animator->GetAnimatorDisplayName());
+			const FName MenuName = Animator->GetAnimatorDisplayName();
 			const FText MenuLabel = FText::FromName(MenuName);
 			const FSlateIcon MenuIcon = FSlateIconFinder::FindIconForClass(Animator->GetClass());
 
@@ -185,8 +188,8 @@ void UE::PropertyAnimatorCoreEditor::Menu::FillDeleteAnimatorSection(UToolMenu*
 		}
 
 		ActorAnimatorsSection.AddMenuEntry(
-			FName(Animator->GetAnimatorDisplayName())
-			, FText::Format(LOCTEXT("DeleteSingleActorAnimator.Label", "Delete {0}"), FText::FromString(Animator->GetAnimatorDisplayName()))
+			Animator->GetAnimatorDisplayName()
+			, FText::Format(LOCTEXT("DeleteSingleActorAnimator.Label", "Delete {0}"), FText::FromName(Animator->GetAnimatorDisplayName()))
 			, LOCTEXT("DeleteSingleActorAnimator.Tooltip", "Delete selected animator")
 			, FSlateIcon()
 			, FUIAction(
@@ -252,8 +255,8 @@ void UE::PropertyAnimatorCoreEditor::Menu::FillEnableAnimatorSection(UToolMenu*
 		}
 
 		ActorAnimatorsSection.AddMenuEntry(
-			FName(Animator->GetAnimatorDisplayName())
-			, FText::Format(LOCTEXT("EnableAnimator.Label", "Enable {0}"), FText::FromString(Animator->GetAnimatorDisplayName()))
+			Animator->GetAnimatorDisplayName()
+			, FText::Format(LOCTEXT("EnableAnimator.Label", "Enable {0}"), FText::FromName(Animator->GetAnimatorDisplayName()))
 			, LOCTEXT("EnableAnimator.Tooltip", "Enable selected animator")
 			, FSlateIconFinder::FindIconForClass(Animator->GetClass())
 			, FUIAction(
@@ -319,8 +322,8 @@ void UE::PropertyAnimatorCoreEditor::Menu::FillDisableAnimatorSection(UToolMenu*
 		}
 
 		ActorAnimatorsSection.AddMenuEntry(
-			FName(Animator->GetAnimatorDisplayName())
-			, FText::Format(LOCTEXT("DisableAnimator.Label", "Disable {0}"), FText::FromString(Animator->GetAnimatorDisplayName()))
+			Animator->GetAnimatorDisplayName()
+			, FText::Format(LOCTEXT("DisableAnimator.Label", "Disable {0}"), FText::FromName(Animator->GetAnimatorDisplayName()))
 			, LOCTEXT("DisableAnimator.Tooltip", "Disable selected animator")
 			, FSlateIconFinder::FindIconForClass(Animator->GetClass())
 			, FUIAction(
