MotionDesign : PropertyAnimatorCore
- Add property aliases to replace original property name to simplify user usage
- Replacing Rotator Roll, Pitch, Yaw with an alias X, Y, Z

diff --git a/Engine/Plugins/Experimental/PropertyAnimatorCore/Source/PropertyAnimatorCore/Private/Animators/PropertyAnimatorCoreBase.cpp b/Engine/Plugins/Experimental/PropertyAnimatorCore/Source/PropertyAnimatorCore/Private/Animators/PropertyAnimatorCoreBase.cpp
--- a/Engine/Plugins/Experimental/PropertyAnimatorCore/Source/PropertyAnimatorCore/Private/Animators/PropertyAnimatorCoreBase.cpp
+++ b/Engine/Plugins/Experimental/PropertyAnimatorCore/Source/PropertyAnimatorCore/Private/Animators/PropertyAnimatorCoreBase.cpp
@@ -53,7 +53,7 @@ void UPropertyAnimatorCoreBase::UpdateAnimatorDisplayName()
 	TArray<FString> PropertiesNames;
 	for (const FPropertyAnimatorCoreData& LinkedProperty : GetLinkedProperties())
 	{
-		PropertiesNames.Add(LinkedProperty.GetPropertyDisplayName().ToString());
+		PropertiesNames.Add(LinkedProperty.GetPropertyDisplayName());
 	}
 
 	auto FindCommonPrefix = [](const TConstArrayView<FString>& InNames)->FString
diff --git a/Engine/Plugins/Experimental/PropertyAnimatorCore/Source/PropertyAnimatorCore/Private/Properties/PropertyAnimatorCoreData.cpp b/Engine/Plugins/Experimental/PropertyAnimatorCore/Source/PropertyAnimatorCore/Private/Properties/PropertyAnimatorCoreData.cpp
--- a/Engine/Plugins/Experimental/PropertyAnimatorCore/Source/PropertyAnimatorCore/Private/Properties/PropertyAnimatorCoreData.cpp
+++ b/Engine/Plugins/Experimental/PropertyAnimatorCore/Source/PropertyAnimatorCore/Private/Properties/PropertyAnimatorCoreData.cpp
@@ -170,6 +170,73 @@ TArray<FString> FPropertyAnimatorCoreData::GetOuterNames() const
 	return OuterNames;
 }
 
+FString FPropertyAnimatorCoreData::GetPropertyDisplayName() const
+{
+	if (!PropertyDisplayName.IsEmpty())
+	{
+		return PropertyDisplayName;
+	}
+
+	const FString ResolverName = IsResolvable()
+		? (GetPropertyResolver()->GetResolverName().ToString() + TEXT("."))
+		: TEXT("");
+
+	FString DisplayName;
+	FString PropertyTypePath;
+	for (const TFieldPath<FProperty>& ChainProperty : ChainProperties)
+	{
+		FString FriendlyName = ChainProperty->GetName();
+
+		if (ChainProperty->IsA<FBoolProperty>())
+		{
+			FriendlyName.RemoveFromStart(TEXT("b"), ESearchCase::Type::CaseSensitive);
+		}
+
+		DisplayName += DisplayName.IsEmpty()
+			? FriendlyName
+			: (TEXT(".") + FriendlyName);
+
+		const FString TypeName = GetPropertyTypeName(ChainProperty.Get()).ToString();
+
+		PropertyTypePath += PropertyTypePath.IsEmpty()
+			? TypeName
+			: (TEXT(".") + TypeName);
+	}
+
+	if (!ChainProperties.IsEmpty())
+	{
+		const FString LeafPropertyName = ChainProperties.Last()->GetName();
+
+		PropertyTypePath += PropertyTypePath.IsEmpty()
+			? LeafPropertyName
+			: (TEXT(".") + LeafPropertyName);
+	}
+
+	FPropertyAnimatorCoreData* MutableThis = const_cast<FPropertyAnimatorCoreData*>(this);
+
+	// Find alias
+	FString AliasName;
+	if (const UPropertyAnimatorCoreSubsystem* AnimatorSubsystem = UPropertyAnimatorCoreSubsystem::Get())
+	{
+		AliasName = AnimatorSubsystem->FindPropertyAlias(PropertyTypePath);
+	}
+
+	if (!AliasName.IsEmpty())
+	{
+		int32 LastPeriodIndex;
+		if (DisplayName.FindLastChar(TEXT('.'), LastPeriodIndex))
+		{
+			DisplayName = DisplayName.Left(LastPeriodIndex + 1);
+		}
+
+		DisplayName += AliasName;
+	}
+
+	MutableThis->PropertyDisplayName = ResolverName + DisplayName;
+
+	return PropertyDisplayName;
+}
+
 FName FPropertyAnimatorCoreData::GetMemberPropertyName() const
 {
 	const FProperty* MemberProperty = GetMemberProperty();
@@ -379,12 +446,6 @@ UPropertyAnimatorCoreHandlerBase* FPropertyAnimatorCoreData::GetPropertyHandler(
 	return PropertyHandler;
 }
 
-FPropertyAnimatorCoreData::FPropertyAnimatorCoreData(const FString& InPathHash, FName InDisplayName)
-	: PropertyDisplayName(InDisplayName)
-	, PathHash(InPathHash)
-{
-}
-
 void FPropertyAnimatorCoreData::GetPropertyValuePtrInternal(void* OutValue) const
 {
 	const FProperty* MemberProperty = GetMemberProperty();
@@ -579,17 +640,7 @@ void FPropertyAnimatorCoreData::GeneratePropertyPath()
 	for (const TFieldPath<FProperty>& ChainProperty : ChainProperties)
 	{
 		PathHash += TEXT(".") + ChainProperty->GetName();
-
-		FString FriendlyName = ChainProperty->GetName();
-		if (ChainProperty->IsA<FBoolProperty>())
-		{
-			FriendlyName.RemoveFromStart(TEXT("b"), ESearchCase::Type::CaseSensitive);
-		}
-
-		DisplayName += DisplayName.IsEmpty() ? FriendlyName : TEXT(".") + FriendlyName;
 	}
-
-	PropertyDisplayName = FName(DisplayName);
 }
 
 bool FPropertyAnimatorCoreData::FindSetterFunctions()
diff --git a/Engine/Plugins/Experimental/PropertyAnimatorCore/Source/PropertyAnimatorCore/Private/Subsystems/PropertyAnimatorCoreSubsystem.cpp b/Engine/Plugins/Experimental/PropertyAnimatorCore/Source/PropertyAnimatorCore/Private/Subsystems/PropertyAnimatorCoreSubsystem.cpp
--- a/Engine/Plugins/Experimental/PropertyAnimatorCore/Source/PropertyAnimatorCore/Private/Subsystems/PropertyAnimatorCoreSubsystem.cpp
+++ b/Engine/Plugins/Experimental/PropertyAnimatorCore/Source/PropertyAnimatorCore/Private/Subsystems/PropertyAnimatorCoreSubsystem.cpp
@@ -42,6 +42,13 @@ void UPropertyAnimatorCoreSubsystem::Initialize(FSubsystemCollectionBase& Collec
 	{
 		return InOwner->FindFunction(TEXT("SetActorHiddenInGame"));
 	});
+
+	// Register alias for Rotator properties
+
+	const FString PropertyType = TEXT("Rotator.double.");
+	RegisterPropertyAlias(PropertyType + GET_MEMBER_NAME_STRING_CHECKED(FRotator, Roll), TEXT("X"));
+	RegisterPropertyAlias(PropertyType + GET_MEMBER_NAME_STRING_CHECKED(FRotator, Pitch), TEXT("Y"));
+	RegisterPropertyAlias(PropertyType + GET_MEMBER_NAME_STRING_CHECKED(FRotator, Yaw), TEXT("Z"));
 }
 
 void UPropertyAnimatorCoreSubsystem::Deinitialize()
@@ -86,9 +93,12 @@ bool UPropertyAnimatorCoreSubsystem::RegisterAnimatorClass(const UClass* InAnima
 
 	if (UPropertyAnimatorCoreBase* CDO = InAnimatorClass->GetDefaultObject<UPropertyAnimatorCoreBase>())
 	{
-		AnimatorsWeak.Add(CDO);
+		if (!CDO->GetAnimatorOriginalName().IsNone())
+		{
+			AnimatorsWeak.Add(CDO);
 
-		return true;
+			return true;
+		}
 	}
 
 	return false;
@@ -1170,6 +1180,38 @@ TSet<UPropertyAnimatorCoreConverterBase*> UPropertyAnimatorCoreSubsystem::GetSup
 	return SupportedConverters;
 }
 
+bool UPropertyAnimatorCoreSubsystem::RegisterPropertyAlias(const FString& InPropertyIdentifier, const FString& InAliasPropertyName)
+{
+	if (InPropertyIdentifier.IsEmpty() || InAliasPropertyName.IsEmpty())
+	{
+		return false;
+	}
+
+	PropertyAliases.Add(InPropertyIdentifier, InAliasPropertyName);
+
+	return true;
+}
+
+bool UPropertyAnimatorCoreSubsystem::UnregisterPropertyAlias(const FString& InPropertyIdentifier)
+{
+	return PropertyAliases.Remove(InPropertyIdentifier) > 0;
+}
+
+FString UPropertyAnimatorCoreSubsystem::FindPropertyAlias(const FString& InPropertyIdentifier) const
+{
+	FString Alias;
+
+	if (!InPropertyIdentifier.IsEmpty())
+	{
+		if (const FString* const AliasPtr = PropertyAliases.Find(InPropertyIdentifier))
+		{
+			Alias = *AliasPtr;
+		}
+	}
+
+	return Alias;
+}
+
 void UPropertyAnimatorCoreSubsystem::SetActorAnimatorsEnabled(const TSet<AActor*>& InActors, bool bInEnabled, bool bInTransact)
 {
 	if (InActors.IsEmpty())
diff --git a/Engine/Plugins/Experimental/PropertyAnimatorCore/Source/PropertyAnimatorCore/Public/Properties/PropertyAnimatorCoreData.h b/Engine/Plugins/Experimental/PropertyAnimatorCore/Source/PropertyAnimatorCore/Public/Properties/PropertyAnimatorCoreData.h
--- a/Engine/Plugins/Experimental/PropertyAnimatorCore/Source/PropertyAnimatorCore/Public/Properties/PropertyAnimatorCoreData.h
+++ b/Engine/Plugins/Experimental/PropertyAnimatorCore/Source/PropertyAnimatorCore/Public/Properties/PropertyAnimatorCoreData.h
@@ -94,11 +94,8 @@ struct FPropertyAnimatorCoreData
 		return !ChainProperties.IsEmpty() ? ChainProperties.Last().Get() : nullptr;
 	}
 
-	/** The full friendly display name from the member property to the inner property */
-	FName GetPropertyDisplayName() const
-	{
-		return PropertyDisplayName;
-	}
+	/** The friendly display name from the member property to the inner property */
+	PROPERTYANIMATORCORE_API FString GetPropertyDisplayName() const;
 
 	/** The full path of a property with its owner */
 	FString GetPathHash() const
@@ -206,9 +203,6 @@ private:
 
 	static FName GetPropertyTypeName(const FProperty* InProperty);
 
-	/** Internal use only, used to quickly identify a data within a controller */
-	explicit FPropertyAnimatorCoreData(const FString& InPathHash, FName InDisplayName);
-
 	/** Internal use only, get property value */
 	void GetPropertyValuePtrInternal(void* OutValue) const;
 
@@ -231,9 +225,9 @@ private:
 	UPROPERTY()
 	TWeakObjectPtr<UObject> OwnerWeak;
 
-	/** The friendly display name of the property we are controlling */
-	UPROPERTY(VisibleInstanceOnly, Category="Property")
-	FName PropertyDisplayName;
+	/** The cached friendly display name of the property we are controlling */
+	UPROPERTY(Transient)
+	FString PropertyDisplayName;
 
 	/** Used to quickly compare struct of this type */
 	UPROPERTY()
diff --git a/Engine/Plugins/Experimental/PropertyAnimatorCore/Source/PropertyAnimatorCore/Public/Subsystems/PropertyAnimatorCoreSubsystem.h b/Engine/Plugins/Experimental/PropertyAnimatorCore/Source/PropertyAnimatorCore/Public/Subsystems/PropertyAnimatorCoreSubsystem.h
--- a/Engine/Plugins/Experimental/PropertyAnimatorCore/Source/PropertyAnimatorCore/Public/Subsystems/PropertyAnimatorCoreSubsystem.h
+++ b/Engine/Plugins/Experimental/PropertyAnimatorCore/Source/PropertyAnimatorCore/Public/Subsystems/PropertyAnimatorCoreSubsystem.h
@@ -143,6 +143,15 @@ public:
 	/** Finds suitable converters for a type conversion */
 	PROPERTYANIMATORCORE_API TSet<UPropertyAnimatorCoreConverterBase*> GetSupportedConverters(const FPropertyBagPropertyDesc& InFromProperty, const FPropertyBagPropertyDesc& InToProperty) const;
 
+	/** Registers a property alias by using a property identifier and property name, property identifier should be like Type.InnerType.PropertyName */
+	bool RegisterPropertyAlias(const FString& InPropertyIdentifier, const FString& InAliasPropertyName);
+
+	/** Unregisters a property alias */
+	bool UnregisterPropertyAlias(const FString& InPropertyIdentifier);
+
+	/** Finds a property alias registered or none */
+	FString FindPropertyAlias(const FString& InPropertyIdentifier) const;
+
 	/** Create an animator of specific class for an actor */
 	PROPERTYANIMATORCORE_API UPropertyAnimatorCoreBase* CreateAnimator(AActor* InActor, const UClass* InAnimatorClass, UPropertyAnimatorCorePresetBase* InPreset = nullptr, bool bInTransact = false) const;
 
@@ -226,4 +235,7 @@ protected:
 
 	/** Some property and their setter cannot be identified automatically, use manual setter resolvers */
 	TMap<FName, TFunction<UFunction*(const UObject*)>> SetterResolvers;
+
+	/** Some property should have a friendlier name and replace the original name by an alias */
+	TMap<FString, FString> PropertyAliases;
 };
diff --git a/Engine/Plugins/Experimental/PropertyAnimatorCore/Source/PropertyAnimatorCoreEditor/Private/Customizations/PropertyAnimatorCoreEditorContextTypeCustomization.cpp b/Engine/Plugins/Experimental/PropertyAnimatorCore/Source/PropertyAnimatorCoreEditor/Private/Customizations/PropertyAnimatorCoreEditorContextTypeCustomization.cpp
--- a/Engine/Plugins/Experimental/PropertyAnimatorCore/Source/PropertyAnimatorCoreEditor/Private/Customizations/PropertyAnimatorCoreEditorContextTypeCustomization.cpp
+++ b/Engine/Plugins/Experimental/PropertyAnimatorCore/Source/PropertyAnimatorCoreEditor/Private/Customizations/PropertyAnimatorCoreEditorContextTypeCustomization.cpp
@@ -34,7 +34,7 @@ void FPropertyAnimatorCoreEditorContextTypeCustomization::CustomizeHeader(TShare
 
 	PropertyContextHandle = InPropertyHandle;
 
-	const FName PropertyName = PropertyContext->GetAnimatedProperty().GetPropertyDisplayName();
+	const FString PropertyDisplayName = PropertyContext->GetAnimatedProperty().GetPropertyDisplayName();
 	const FName PropertyTypeName = PropertyContext->GetAnimatedProperty().GetLeafPropertyTypeName();
 
 	InRow
@@ -51,7 +51,7 @@ void FPropertyAnimatorCoreEditorContextTypeCustomization::CustomizeHeader(TShare
 			+ SHorizontalBox::Slot()
 			.FillWidth(1.f)
 			[
-				InPropertyHandle->CreatePropertyNameWidget(FText::FromString(PropertyName.ToString() + TEXT(" (") + PropertyTypeName.ToString() + TEXT(")")))
+				InPropertyHandle->CreatePropertyNameWidget(FText::FromString(PropertyDisplayName + TEXT(" (") + PropertyTypeName.ToString() + TEXT(")")))
 			]
 		]
 		.ValueContent()
diff --git a/Engine/Plugins/Experimental/PropertyAnimatorCore/Source/PropertyAnimatorCoreEditor/Private/Menus/PropertyAnimatorCoreEditorMenu.cpp b/Engine/Plugins/Experimental/PropertyAnimatorCore/Source/PropertyAnimatorCoreEditor/Private/Menus/PropertyAnimatorCoreEditorMenu.cpp
--- a/Engine/Plugins/Experimental/PropertyAnimatorCore/Source/PropertyAnimatorCoreEditor/Private/Menus/PropertyAnimatorCoreEditorMenu.cpp
+++ b/Engine/Plugins/Experimental/PropertyAnimatorCore/Source/PropertyAnimatorCoreEditor/Private/Menus/PropertyAnimatorCoreEditorMenu.cpp
@@ -688,7 +688,7 @@ void UE::PropertyAnimatorCoreEditor::Menu::FillNewAnimatorSubmenu(UToolMenu* InM
 
 		for (const FPropertyAnimatorCoreData& SupportedProperty : SupportedProperties)
 		{
-			const FString MenuName = SupportedProperty.GetPropertyDisplayName().ToString();
+			const FString MenuName = SupportedProperty.GetPropertyDisplayName();
 			const FText MenuLabel = FText::FromString(MenuName + TEXT(" (") + SupportedProperty.GetLeafPropertyTypeName().ToString() + TEXT(")"));
 
 			// Create action (creates an animator and links the property)
@@ -783,11 +783,11 @@ void UE::PropertyAnimatorCoreEditor::Menu::FillLinkAnimatorSubmenu(UToolMenu* In
 
 	for (const FPropertyAnimatorCoreData& SupportedProperty : SupportedProperties)
 	{
-		const FName MenuName = SupportedProperty.GetPropertyDisplayName();
-		const FText MenuLabel = FText::FromString(MenuName.ToString() + TEXT(" (") + SupportedProperty.GetLeafPropertyTypeName().ToString() + TEXT(")"));
+		const FString MenuName = SupportedProperty.GetPropertyDisplayName();
+		const FText MenuLabel = FText::FromString(MenuName + TEXT(" (") + SupportedProperty.GetLeafPropertyTypeName().ToString() + TEXT(")"));
 
 		PropertySection.AddMenuEntry(
-			MenuName
+			FName(MenuName)
 			, MenuLabel
 			, LOCTEXT("LinkAnimatorPropertySection.Tooltip", "Link or unlink this property from the animator")
 			, FSlateIcon()
@@ -835,12 +835,12 @@ void UE::PropertyAnimatorCoreEditor::Menu::FillPresetAnimatorSubmenu(UToolMenu*
 
 	for (const FPropertyAnimatorCoreData& SupportedProperty : SupportedProperties)
 	{
-		const FName MenuName = SupportedProperty.GetPropertyDisplayName();
-		const FText MenuLabel = FText::FromString(MenuName.ToString() + TEXT(" (") + SupportedProperty.GetLeafPropertyTypeName().ToString() + TEXT(")"));
+		const FString MenuName = SupportedProperty.GetPropertyDisplayName();
+		const FText MenuLabel = FText::FromString(MenuName + TEXT(" (") + SupportedProperty.GetLeafPropertyTypeName().ToString() + TEXT(")"));
 		const FText MenuTooltip = LOCTEXT("LinkPresetProperty.Tooltip", "Link this preset property");
 
 		const FToolMenuEntry SupportedPropertyEntry = FToolMenuEntry::InitMenuEntry(
-			MenuName
+			FName(MenuName)
 			, MenuLabel
 			, MenuTooltip
 			, FSlateIcon()
@@ -926,10 +926,11 @@ void UE::PropertyAnimatorCoreEditor::Menu::FillNewPresetAnimatorSubmenu(UToolMen
 
 	for (const FPropertyAnimatorCoreData& PresetProperty : PresetProperties)
 	{
-		const FText MenuLabel = FText::FromString(PresetProperty.GetPropertyDisplayName().ToString() + TEXT(" (") + PresetProperty.GetLeafPropertyTypeName().ToString() + TEXT(")"));
+		const FString MenuName = PresetProperty.GetPropertyDisplayName();
+		const FText MenuLabel = FText::FromString(MenuName + TEXT(" (") + PresetProperty.GetLeafPropertyTypeName().ToString() + TEXT(")"));
 
 		const FToolMenuEntry CreatePropertyEntry = FToolMenuEntry::InitMenuEntry(
-			FName(TEXT("CreateProperty") + PresetProperty.GetPropertyDisplayName().ToString())
+			FName(TEXT("CreateProperty") + MenuName)
 			, MenuLabel
 			, LOCTEXT("CreateAnimatorPresetProperty.Tooltip", "Create this animator using this preset property")
 			, FSlateIcon()
@@ -948,7 +949,7 @@ void UE::PropertyAnimatorCoreEditor::Menu::FillNewPresetAnimatorSubmenu(UToolMen
 		);
 
 		const FToolMenuEntry LinkPropertyEntry = FToolMenuEntry::InitMenuEntry(
-			FName(TEXT("LinkProperty") + PresetProperty.GetPropertyDisplayName().ToString())
+			FName(TEXT("LinkProperty") + MenuName)
 			, MenuLabel
 			, LOCTEXT("LinkAnimatorPresetProperty.Tooltip", "Link this preset property to the last created animator")
 			, FSlateIcon()
