MotionDesign : PropertyAnimatorCore
- Replacing animator channel track by an animator track that binds to the time source object it is controlling
- Updating code to use binding object to update time and restore state
- Adding section options to select time eval mode (sequence, section, custom) to suit different needs

diff --git a/Engine/Plugins/Experimental/PropertyAnimatorCore/Source/PropertyAnimatorCore/Private/Sequencer/MovieSceneAnimatorExecutionToken.h b/Engine/Plugins/Experimental/PropertyAnimatorCore/Source/PropertyAnimatorCore/Private/Sequencer/MovieSceneAnimatorExecutionToken.h
--- a/Engine/Plugins/Experimental/PropertyAnimatorCore/Source/PropertyAnimatorCore/Private/Sequencer/MovieSceneAnimatorExecutionToken.h
+++ b/Engine/Plugins/Experimental/PropertyAnimatorCore/Source/PropertyAnimatorCore/Private/Sequencer/MovieSceneAnimatorExecutionToken.h
@@ -8,43 +8,38 @@
 #include "TimeSources/PropertyAnimatorCoreSequencerTimeSource.h"
 
 /** Used to restore state back to previous when outside section */
-struct FMovieSceneAnimatorPreAnimatedTokenProducer : IMovieScenePreAnimatedGlobalTokenProducer
+struct FMovieSceneAnimatorPreAnimatedTokenProducer : IMovieScenePreAnimatedTokenProducer
 {
-	explicit FMovieSceneAnimatorPreAnimatedTokenProducer(uint8 InChannel)
-		: Channel(InChannel)
+	FMovieSceneAnimatorPreAnimatedTokenProducer()
 	{}
 
-	//~ Begin IMovieScenePreAnimatedGlobalTokenProducer
-	virtual IMovieScenePreAnimatedGlobalTokenPtr CacheExistingState() const override
+	//~ Begin IMovieScenePreAnimatedTokenProducer
+	virtual IMovieScenePreAnimatedTokenPtr CacheExistingState(UObject& InObject) const override
 	{
-		struct FMovieSceneAnimatorPreAnimatedGlobalToken : IMovieScenePreAnimatedGlobalToken
+		struct FMovieSceneAnimatorPreAnimatedToken : IMovieScenePreAnimatedToken
 		{
-			FMovieSceneAnimatorPreAnimatedGlobalToken(uint8 InChannel)
-				: Channel(InChannel)
+			FMovieSceneAnimatorPreAnimatedToken()
 			{}
 
-			//~ Begin IMovieScenePreAnimatedGlobalToken
-			virtual void RestoreState(const UE::MovieScene::FRestoreStateParams& InParams) override
+			//~ Begin IMovieScenePreAnimatedToken
+			virtual void RestoreState(UObject& InObject, const UE::MovieScene::FRestoreStateParams& InParams) override
 			{
-				UPropertyAnimatorCoreSequencerTimeSource::OnAnimatorTimeEvaluated.Broadcast(Channel, TOptional<double>(), TOptional<float>());
+				if (UPropertyAnimatorCoreSequencerTimeSource* SequencerTimeSource = Cast<UPropertyAnimatorCoreSequencerTimeSource>(&InObject))
+				{
+					SequencerTimeSource->OnSequencerTimeEvaluated(TOptional<double>(), TOptional<float>());
+				}
 			}
-			//~ End IMovieScenePreAnimatedGlobalToken
-
-		private:
-			uint8 Channel = 0;
+			//~ End IMovieScenePreAnimatedToken
 		};
 
-		return FMovieSceneAnimatorPreAnimatedGlobalToken(Channel);
+		return FMovieSceneAnimatorPreAnimatedToken();
 	}
-	//~ End IMovieScenePreAnimatedGlobalTokenProducer
+	//~ End IMovieScenePreAnimatedTokenProducer
 
 	static FMovieSceneAnimTypeID GetAnimTypeID()
 	{
 		return TMovieSceneAnimTypeID<FMovieSceneAnimatorPreAnimatedTokenProducer>();
 	}
-
-private:
-	uint8 Channel = 0;
 };
 
 /** Used to evaluate active section */
@@ -57,20 +52,44 @@ struct FMovieSceneAnimatorExecutionToken : IMovieSceneExecutionToken
 	//~ Begin IMovieSceneExecutionToken
 	virtual void Execute(const FMovieSceneContext& InContext, const FMovieSceneEvaluationOperand& InOperand, FPersistentEvaluationData& InPersistentData, IMovieScenePlayer& InPlayer) override
 	{
-		InPlayer.SavePreAnimatedState(FMovieSceneAnimatorPreAnimatedTokenProducer::GetAnimTypeID(), FMovieSceneAnimatorPreAnimatedTokenProducer(SectionData.Channel));
+		if (InOperand.ObjectBindingID.IsValid() && SectionData.Section)
+		{
+			double EvaluatedTime = InContext.GetTime().AsDecimal();
 
-		double EvaluatedTime = InContext.GetTime().AsDecimal();
+			if (SectionData.EvalTimeMode == EMovieSceneAnimatorEvalTimeMode::Section)
+			{
+				EvaluatedTime -= SectionData.Section->GetInclusiveStartFrame().Value;
+			}
 
-		if (SectionData.bUseSectionTime && SectionData.Section)
-		{
-			EvaluatedTime -= SectionData.Section->GetInclusiveStartFrame().Value;
-		}
+			EvaluatedTime /= InContext.GetFrameRate().AsDecimal();
 
-		EvaluatedTime /= InContext.GetFrameRate().AsDecimal();
+			if (SectionData.EvalTimeMode == EMovieSceneAnimatorEvalTimeMode::Custom)
+			{
+				const double SectionStartTime = (SectionData.Section->GetInclusiveStartFrame().Value * 1.f) / InContext.GetFrameRate().AsDecimal();
+				const double SectionEndTime = (SectionData.Section->GetExclusiveEndFrame().Value * 1.f) / InContext.GetFrameRate().AsDecimal();
 
-		const float Magnitude = SectionData.Section ? SectionData.Section->EvaluateEasing(InContext.GetTime()) : 1;
+				EvaluatedTime = FMath::GetMappedRangeValueClamped(FVector2D(SectionStartTime, SectionEndTime), FVector2D(SectionData.CustomStartTime, SectionData.CustomEndTime), EvaluatedTime);
+			}
+
+			const float Magnitude = SectionData.Section->EvaluateEasing(InContext.GetTime());
+
+			for (const TWeakObjectPtr<UObject>& BoundObjectWeak : InPlayer.FindBoundObjects(InOperand))
+			{
+				UObject* BoundObject = BoundObjectWeak.Get();
 
-		UPropertyAnimatorCoreSequencerTimeSource::OnAnimatorTimeEvaluated.Broadcast(SectionData.Channel, EvaluatedTime, Magnitude);
+				if (!BoundObject)
+				{
+					continue;
+				}
+
+				if (UPropertyAnimatorCoreSequencerTimeSource* SequencerTimeSource = Cast<UPropertyAnimatorCoreSequencerTimeSource>(BoundObject))
+				{
+					InPlayer.SavePreAnimatedState(*SequencerTimeSource, FMovieSceneAnimatorPreAnimatedTokenProducer::GetAnimTypeID(), FMovieSceneAnimatorPreAnimatedTokenProducer());
+
+					SequencerTimeSource->OnSequencerTimeEvaluated(EvaluatedTime, Magnitude);
+				}
+			}
+		}
 	}
 	//~ End IMovieSceneExecutionToken
 
diff --git a/Engine/Plugins/Experimental/PropertyAnimatorCore/Source/PropertyAnimatorCore/Private/Sequencer/MovieSceneAnimatorSection.cpp b/Engine/Plugins/Experimental/PropertyAnimatorCore/Source/PropertyAnimatorCore/Private/Sequencer/MovieSceneAnimatorSection.cpp
--- a/Engine/Plugins/Experimental/PropertyAnimatorCore/Source/PropertyAnimatorCore/Private/Sequencer/MovieSceneAnimatorSection.cpp
+++ b/Engine/Plugins/Experimental/PropertyAnimatorCore/Source/PropertyAnimatorCore/Private/Sequencer/MovieSceneAnimatorSection.cpp
@@ -9,12 +9,17 @@ UMovieSceneAnimatorSection::UMovieSceneAnimatorSection()
 	EvalOptions.CompletionMode = EMovieSceneCompletionMode::RestoreState;
 }
 
-void UMovieSceneAnimatorSection::SetChannel(uint8 InChannel)
+void UMovieSceneAnimatorSection::SetEvalTimeMode(EMovieSceneAnimatorEvalTimeMode InMode)
 {
-	Channel = InChannel;
+	EvalTimeMode = InMode;
 }
 
-void UMovieSceneAnimatorSection::SetUseSectionTime(bool bInUse)
+void UMovieSceneAnimatorSection::SetCustomStartTime(double InTime)
 {
-	bUseSectionTime = bInUse;
+	CustomStartTime = InTime;
+}
+
+void UMovieSceneAnimatorSection::SetCustomEndTime(double InTime)
+{
+	CustomEndTime = InTime;
 }
diff --git a/Engine/Plugins/Experimental/PropertyAnimatorCore/Source/PropertyAnimatorCore/Private/Sequencer/MovieSceneAnimatorTrack.cpp b/Engine/Plugins/Experimental/PropertyAnimatorCore/Source/PropertyAnimatorCore/Private/Sequencer/MovieSceneAnimatorTrack.cpp
--- a/Engine/Plugins/Experimental/PropertyAnimatorCore/Source/PropertyAnimatorCore/Private/Sequencer/MovieSceneAnimatorTrack.cpp
+++ b/Engine/Plugins/Experimental/PropertyAnimatorCore/Source/PropertyAnimatorCore/Private/Sequencer/MovieSceneAnimatorTrack.cpp
@@ -17,24 +17,6 @@ UMovieSceneAnimatorTrack::UMovieSceneAnimatorTrack()
 	SupportedBlendTypes.Add(EMovieSceneBlendType::Absolute);
 }
 
-int32 UMovieSceneAnimatorTrack::GetChannelCount(uint8 InChannel) const
-{
-	int32 Count = 0;
-
-	for (UMovieSceneSection* Section : GetAllSections())
-	{
-		if (const UMovieSceneAnimatorSection* AnimatorSection = Cast<UMovieSceneAnimatorSection>(Section))
-		{
-			if (AnimatorSection->GetChannel() == InChannel)
-			{
-				Count++;
-			}
-		}
-	}
-
-	return Count;
-}
-
 bool UMovieSceneAnimatorTrack::SupportsType(TSubclassOf<UMovieSceneSection> InSectionClass) const
 {
 	return InSectionClass == UMovieSceneAnimatorSection::StaticClass();
@@ -95,20 +77,7 @@ void UMovieSceneAnimatorTrack::RemoveSectionAt(int32 InSectionIndex)
 #if WITH_EDITORONLY_DATA
 FText UMovieSceneAnimatorTrack::GetDefaultDisplayName() const
 {
-	if (!Sections.IsEmpty() && Sections[0])
-	{
-		if (const UMovieSceneAnimatorSection* AnimatorSection = Cast<UMovieSceneAnimatorSection>(Sections[0]))
-		{
-			return FText::Format(LOCTEXT("MovieSceneAnimatorTrackChannelName", "Animator Channel {0}"), FText::AsNumber(AnimatorSection->GetChannel()));
-		}
-	}
-
-	return LOCTEXT("MovieSceneAnimatorTrackName", "Animator Channel");
-}
-
-bool UMovieSceneAnimatorTrack::CanRename() const
-{
-	return false;
+	return LOCTEXT("MovieSceneAnimatorTrackDefaultName", "Animator Track");
 }
 #endif
 
@@ -122,8 +91,9 @@ FMovieSceneEvalTemplatePtr UMovieSceneAnimatorTrack::CreateTemplateForSection(co
 	}
 
 	FMovieSceneAnimatorSectionData SectionData;
-	SectionData.Channel = AnimatorSection->GetChannel();
-	SectionData.bUseSectionTime = AnimatorSection->GetUseSectionTime();
+	SectionData.EvalTimeMode = AnimatorSection->GetEvalTimeMode();
+	SectionData.CustomStartTime = AnimatorSection->GetCustomStartTime();
+	SectionData.CustomEndTime = AnimatorSection->GetCustomEndTime();
 	SectionData.Section = AnimatorSection;
 
 	return FMovieSceneAnimatorEvalTemplate(SectionData);
diff --git a/Engine/Plugins/Experimental/PropertyAnimatorCore/Source/PropertyAnimatorCore/Private/TimeSources/PropertyAnimatorCoreSequencerTimeSource.cpp b/Engine/Plugins/Experimental/PropertyAnimatorCore/Source/PropertyAnimatorCore/Private/TimeSources/PropertyAnimatorCoreSequencerTimeSource.cpp
--- a/Engine/Plugins/Experimental/PropertyAnimatorCore/Source/PropertyAnimatorCore/Private/TimeSources/PropertyAnimatorCoreSequencerTimeSource.cpp
+++ b/Engine/Plugins/Experimental/PropertyAnimatorCore/Source/PropertyAnimatorCore/Private/TimeSources/PropertyAnimatorCoreSequencerTimeSource.cpp
@@ -2,45 +2,22 @@
 
 #include "TimeSources/PropertyAnimatorCoreSequencerTimeSource.h"
 
-UPropertyAnimatorCoreSequencerTimeSource::FOnAnimatorTimeEvaluated UPropertyAnimatorCoreSequencerTimeSource::OnAnimatorTimeEvaluated;
-
 bool UPropertyAnimatorCoreSequencerTimeSource::UpdateEvaluationData(FPropertyAnimatorCoreTimeSourceEvaluationData& OutData)
 {
-	if (!EvalTime.IsSet() || !EvalMagnitude.IsSet())
+	if (!EvalResult.bEvalValid)
 	{
 		return false;
 	}
 
-	OutData.TimeElapsed = EvalTime.GetValue();
-	OutData.Magnitude = EvalMagnitude.GetValue();
+	OutData.TimeElapsed = EvalResult.EvalTime;
+	OutData.Magnitude = EvalResult.EvalMagnitude;
 
 	return true;
 }
 
-void UPropertyAnimatorCoreSequencerTimeSource::OnTimeSourceActive()
-{
-	Super::OnTimeSourceActive();
-
-	OnAnimatorTimeEvaluated.AddUObject(this, &UPropertyAnimatorCoreSequencerTimeSource::OnSequencerTimeEvaluated);
-}
-
-void UPropertyAnimatorCoreSequencerTimeSource::OnTimeSourceInactive()
-{
-	Super::OnTimeSourceInactive();
-
-	OnAnimatorTimeEvaluated.RemoveAll(this);
-}
-
-void UPropertyAnimatorCoreSequencerTimeSource::SetChannel(uint8 InChannel)
+void UPropertyAnimatorCoreSequencerTimeSource::OnSequencerTimeEvaluated(const TOptional<double>& InTimeEval, const TOptional<float>& InMagnitudeEval)
 {
-	ChannelData.Channel = InChannel;
-}
-
-void UPropertyAnimatorCoreSequencerTimeSource::OnSequencerTimeEvaluated(uint8 InChannel, const TOptional<double>& InTimeEval, const TOptional<float>& InMagnitudeEval)
-{
-	if (ChannelData.Channel == InChannel)
-	{
-		EvalTime = InTimeEval;
-		EvalMagnitude = InMagnitudeEval;
-	}
+	EvalResult.bEvalValid = InTimeEval.IsSet() && InMagnitudeEval.IsSet();
+	EvalResult.EvalTime = InTimeEval.Get(0.0);
+	EvalResult.EvalMagnitude = InMagnitudeEval.Get(1.f);
 }
diff --git a/Engine/Plugins/Experimental/PropertyAnimatorCore/Source/PropertyAnimatorCore/Public/Sequencer/MovieSceneAnimatorSection.h b/Engine/Plugins/Experimental/PropertyAnimatorCore/Source/PropertyAnimatorCore/Public/Sequencer/MovieSceneAnimatorSection.h
--- a/Engine/Plugins/Experimental/PropertyAnimatorCore/Source/PropertyAnimatorCore/Public/Sequencer/MovieSceneAnimatorSection.h
+++ b/Engine/Plugins/Experimental/PropertyAnimatorCore/Source/PropertyAnimatorCore/Public/Sequencer/MovieSceneAnimatorSection.h
@@ -6,7 +6,7 @@
 #include "MovieSceneSection.h"
 #include "MovieSceneAnimatorSection.generated.h"
 
-/** Movie scene section for a sequencer animator channel */
+/** Movie scene section for a sequencer animator track */
 UCLASS(MinimalAPI)
 class UMovieSceneAnimatorSection : public UMovieSceneSection
 {
@@ -15,26 +15,34 @@ class UMovieSceneAnimatorSection : public UMovieSceneSection
 public:
 	UMovieSceneAnimatorSection();
 
-	PROPERTYANIMATORCORE_API void SetChannel(uint8 InChannel);
+	PROPERTYANIMATORCORE_API void SetEvalTimeMode(EMovieSceneAnimatorEvalTimeMode InMode);
 
-	uint8 GetChannel() const
+	EMovieSceneAnimatorEvalTimeMode GetEvalTimeMode() const
 	{
-		return Channel;
+		return EvalTimeMode;
 	}
 
-	PROPERTYANIMATORCORE_API void SetUseSectionTime(bool bInUse);
+	PROPERTYANIMATORCORE_API void SetCustomStartTime(double InTime);
 
-	bool GetUseSectionTime() const
+	double GetCustomStartTime() const
 	{
-		return bUseSectionTime;
+		return CustomStartTime;
+	}
+
+	PROPERTYANIMATORCORE_API void SetCustomEndTime(double InTime);
+
+	double GetCustomEndTime() const
+	{
+		return CustomEndTime;
 	}
 
 protected:
-	/** Channel used to push sequencer time to */
-	UPROPERTY(EditAnywhere, Category="Animator")
-	uint8 Channel = 0;
+	UPROPERTY(EditAnywhere, Setter, Getter, Category="Animator")
+	EMovieSceneAnimatorEvalTimeMode EvalTimeMode = EMovieSceneAnimatorEvalTimeMode::Section;
+
+	UPROPERTY(EditAnywhere, Setter, Getter, Category="Animator", meta=(EditCondition="EvalTimeMode == EMovieSceneAnimatorEvalTimeMode::Custom"))
+	double CustomStartTime = 0.0;
 
-	/** Whether to use the section time (playback range) or track time (sequencer time) */
-	UPROPERTY(EditAnywhere, Category="Animator")
-	bool bUseSectionTime = true;
+	UPROPERTY(EditAnywhere, Setter, Getter, Category="Animator", meta=(EditCondition="EvalTimeMode == EMovieSceneAnimatorEvalTimeMode::Custom"))
+	double CustomEndTime = 1.0;
 };
diff --git a/Engine/Plugins/Experimental/PropertyAnimatorCore/Source/PropertyAnimatorCore/Public/Sequencer/MovieSceneAnimatorTrack.h b/Engine/Plugins/Experimental/PropertyAnimatorCore/Source/PropertyAnimatorCore/Public/Sequencer/MovieSceneAnimatorTrack.h
--- a/Engine/Plugins/Experimental/PropertyAnimatorCore/Source/PropertyAnimatorCore/Public/Sequencer/MovieSceneAnimatorTrack.h
+++ b/Engine/Plugins/Experimental/PropertyAnimatorCore/Source/PropertyAnimatorCore/Public/Sequencer/MovieSceneAnimatorTrack.h
@@ -15,9 +15,6 @@ class UMovieSceneAnimatorTrack : public UMovieSceneNameableTrack, public IMovieS
 public:
 	UMovieSceneAnimatorTrack();
 
-	/** Get the number of section using the channel under this track */
-	PROPERTYANIMATORCORE_API int32 GetChannelCount(uint8 InChannel) const;
-
 	//~ Begin UMovieSceneTrack
 	virtual bool SupportsType(TSubclassOf<UMovieSceneSection> InSectionClass) const override;
 	virtual UMovieSceneSection* CreateNewSection() override;
@@ -30,7 +27,6 @@ public:
 	virtual void RemoveSectionAt(int32 InSectionIndex) override;
 #if WITH_EDITORONLY_DATA
 	virtual FText GetDefaultDisplayName() const override;
-	virtual bool CanRename() const override;
 #endif
 	//~ End UMovieSceneTrack
 
diff --git a/Engine/Plugins/Experimental/PropertyAnimatorCore/Source/PropertyAnimatorCore/Public/Sequencer/MovieSceneAnimatorTypes.h b/Engine/Plugins/Experimental/PropertyAnimatorCore/Source/PropertyAnimatorCore/Public/Sequencer/MovieSceneAnimatorTypes.h
--- a/Engine/Plugins/Experimental/PropertyAnimatorCore/Source/PropertyAnimatorCore/Public/Sequencer/MovieSceneAnimatorTypes.h
+++ b/Engine/Plugins/Experimental/PropertyAnimatorCore/Source/PropertyAnimatorCore/Public/Sequencer/MovieSceneAnimatorTypes.h
@@ -3,16 +3,29 @@
 #pragma once
 
 #include "HAL/Platform.h"
+#include "MovieSceneAnimatorTypes.generated.h"
 
 class UMovieSceneAnimatorSection;
 
+/** Enumerates all possible ways of interpreting time */
+UENUM()
+enum class EMovieSceneAnimatorEvalTimeMode : uint8
+{
+	/** The sequence time, if the section has an offset, it won't matter */
+	Sequence,
+	/** The section time, takes into account any offset the section has */
+	Section,
+	/** Custom provided time, interpolate between start and end time based on progress */
+	Custom
+};
+
 struct FMovieSceneAnimatorSectionData
 {
-	/** Channel used to push sequencer time to */
-	uint8 Channel = 0;
+	EMovieSceneAnimatorEvalTimeMode EvalTimeMode = EMovieSceneAnimatorEvalTimeMode::Sequence;
+
+	double CustomStartTime;
 
-	/** Whether to use the section time or track time (sequencer time) */
-	bool bUseSectionTime = true;
+	double CustomEndTime;
 
 	/** The section object to get easing from */
 	const UMovieSceneAnimatorSection* Section = nullptr;
diff --git a/Engine/Plugins/Experimental/PropertyAnimatorCore/Source/PropertyAnimatorCore/Public/TimeSources/PropertyAnimatorCoreManualTimeSource.h b/Engine/Plugins/Experimental/PropertyAnimatorCore/Source/PropertyAnimatorCore/Public/TimeSources/PropertyAnimatorCoreManualTimeSource.h
--- a/Engine/Plugins/Experimental/PropertyAnimatorCore/Source/PropertyAnimatorCore/Public/TimeSources/PropertyAnimatorCoreManualTimeSource.h
+++ b/Engine/Plugins/Experimental/PropertyAnimatorCore/Source/PropertyAnimatorCore/Public/TimeSources/PropertyAnimatorCoreManualTimeSource.h
@@ -84,7 +84,7 @@ protected:
 	bool bOverrideTime = true;
 
 	/** Time to evaluate */
-	UPROPERTY(EditInstanceOnly, Setter, Getter, Category="Animator", meta=(EditCondition="bOverrideTime"))
+	UPROPERTY(EditInstanceOnly, Setter, Getter, Category="Animator", meta=(EditCondition="bOverrideTime", Units=Seconds))
 	double CustomTime = 0.f;
 
 	/** Playback state */
diff --git a/Engine/Plugins/Experimental/PropertyAnimatorCore/Source/PropertyAnimatorCore/Public/TimeSources/PropertyAnimatorCoreSequencerTimeSource.h b/Engine/Plugins/Experimental/PropertyAnimatorCore/Source/PropertyAnimatorCore/Public/TimeSources/PropertyAnimatorCoreSequencerTimeSource.h
--- a/Engine/Plugins/Experimental/PropertyAnimatorCore/Source/PropertyAnimatorCore/Public/TimeSources/PropertyAnimatorCoreSequencerTimeSource.h
+++ b/Engine/Plugins/Experimental/PropertyAnimatorCore/Source/PropertyAnimatorCore/Public/TimeSources/PropertyAnimatorCoreSequencerTimeSource.h
@@ -6,50 +6,39 @@
 #include "PropertyAnimatorCoreSequencerTimeSource.generated.h"
 
 USTRUCT()
-struct FPropertyAnimatorCoreSequencerTimeSourceChannel
+struct FPropertyAnimatorCoreSequencerTimeSourceEvalResult
 {
 	GENERATED_BODY()
 
-	UPROPERTY(EditInstanceOnly, Category="Animator")
-	uint8 Channel = 0;
+	/** Is the evaluation state valid */
+	bool bEvalValid = false;
+
+	/** Last evaluated time received */
+	UPROPERTY(VisibleInstanceOnly, Transient, Category="Animator", meta=(Units=Seconds))
+	double EvalTime = 0.0;
+
+	/** Last evaluated magnitude received */
+	float EvalMagnitude = 1.f;
 };
 
-/** Sequencer time source that sync with animator track channel */
+/** Time source that sync with a sequencer animator track */
 UCLASS(MinimalAPI)
 class UPropertyAnimatorCoreSequencerTimeSource : public UPropertyAnimatorCoreTimeSourceBase
 {
 	GENERATED_BODY()
 
 public:
-	DECLARE_MULTICAST_DELEGATE_ThreeParams(FOnAnimatorTimeEvaluated, uint8 /** Channel */, const TOptional<double>& /** EvalTime */, const TOptional<float>& /** EvalMagnitude */)
-	static FOnAnimatorTimeEvaluated OnAnimatorTimeEvaluated;
-
 	UPropertyAnimatorCoreSequencerTimeSource()
 		: UPropertyAnimatorCoreTimeSourceBase(TEXT("Sequencer"))
 	{}
 
 	//~ Begin UPropertyAnimatorTimeSourceBase
 	virtual bool UpdateEvaluationData(FPropertyAnimatorCoreTimeSourceEvaluationData& OutData) override;
-	virtual void OnTimeSourceActive() override;
-	virtual void OnTimeSourceInactive() override;
 	//~ End UPropertyAnimatorTimeSourceBase
 
-	void SetChannel(uint8 InChannel);
-	uint8 GetChannel() const
-	{
-		return ChannelData.Channel;
-	}
+	void OnSequencerTimeEvaluated(const TOptional<double>& InTimeEval, const TOptional<float>& InMagnitudeEval);
 
 protected:
-	void OnSequencerTimeEvaluated(uint8 InChannel, const TOptional<double>& InTimeEval, const TOptional<float>& InMagnitudeEval);
-
-	/** Channel to sample time from */
-	UPROPERTY(EditInstanceOnly, DisplayName="Channel", Category="Animator")
-	FPropertyAnimatorCoreSequencerTimeSourceChannel ChannelData;
-
-	/** Last evaluated time received */
-	TOptional<double> EvalTime;
-
-	/** Last evaluated magnitude received */
-	TOptional<float> EvalMagnitude;
+	UPROPERTY(EditInstanceOnly, DisplayName="EvalTime", Category="Animator")
+	FPropertyAnimatorCoreSequencerTimeSourceEvalResult EvalResult;
 };
\ No newline at end of file
diff --git a/Engine/Plugins/Experimental/PropertyAnimatorCore/Source/PropertyAnimatorCoreEditor/Private/Customizations/PropertyAnimatorCoreEditorSequencerTimeSourceChannelTypeCustomization.cpp b/Engine/Plugins/Experimental/PropertyAnimatorCore/Source/PropertyAnimatorCoreEditor/Private/Customizations/PropertyAnimatorCoreEditorSequencerTimeSourceChannelTypeCustomization.cpp
deleted file mode 100644
diff --git a/Engine/Plugins/Experimental/PropertyAnimatorCore/Source/PropertyAnimatorCoreEditor/Private/Customizations/PropertyAnimatorCoreEditorSequencerTimeSourceChannelTypeCustomization.h b/Engine/Plugins/Experimental/PropertyAnimatorCore/Source/PropertyAnimatorCoreEditor/Private/Customizations/PropertyAnimatorCoreEditorSequencerTimeSourceChannelTypeCustomization.h
deleted file mode 100644
diff --git a/Engine/Plugins/Experimental/PropertyAnimatorCore/Source/PropertyAnimatorCoreEditor/Private/Customizations/PropertyAnimatorCoreEditorSequencerTimeSourceEvalResultTypeCustomization.cpp b/Engine/Plugins/Experimental/PropertyAnimatorCore/Source/PropertyAnimatorCoreEditor/Private/Customizations/PropertyAnimatorCoreEditorSequencerTimeSourceEvalResultTypeCustomization.cpp
new file mode 100644
--- /dev/null
+++ b/Engine/Plugins/Experimental/PropertyAnimatorCore/Source/PropertyAnimatorCoreEditor/Private/Customizations/PropertyAnimatorCoreEditorSequencerTimeSourceEvalResultTypeCustomization.cpp
@@ -0,0 +1,143 @@
+// Copyright Epic Games, Inc. All Rights Reserved.
+
+#include "Customizations/PropertyAnimatorCoreEditorSequencerTimeSourceEvalResultTypeCustomization.h"
+
+#include "DetailWidgetRow.h"
+#include "Editor.h"
+#include "Selection.h"
+#include "Sequencer/MovieSceneAnimatorTrackEditor.h"
+#include "TimeSources/PropertyAnimatorCoreSequencerTimeSource.h"
+#include "Widgets/Input/SButton.h"
+#include "Widgets/SBoxPanel.h"
+#include "Widgets/Text/STextBlock.h"
+
+#define LOCTEXT_NAMESPACE "PropertyAnimatorCoreEditorSequencerTimeSourceEvalResultTypeCustomization"
+
+void FPropertyAnimatorCoreEditorSequencerTimeSourceChannelTypeCustomization::CustomizeHeader(TSharedRef<IPropertyHandle> InPropertyHandle, FDetailWidgetRow& InRow, IPropertyTypeCustomizationUtils& InUtils)
+{
+	if (!InPropertyHandle->IsValidHandle())
+	{
+		return;
+	}
+
+	EvalTimePropertyHandle = InPropertyHandle->GetChildHandle(GET_MEMBER_NAME_CHECKED(FPropertyAnimatorCoreSequencerTimeSourceEvalResult, EvalTime));
+
+	if (!EvalTimePropertyHandle->IsValidHandle())
+	{
+		return;
+	}
+
+	InRow.NameContent()
+	[
+		InPropertyHandle->CreatePropertyNameWidget()
+	];
+
+	InRow.ValueContent()
+	.HAlign(HAlign_Fill)
+	[
+		SNew(SHorizontalBox)
+		+ SHorizontalBox::Slot()
+		.AutoWidth()
+		.HAlign(HAlign_Fill)
+		[
+			EvalTimePropertyHandle->CreatePropertyValueWidget()
+		]
+		+ SHorizontalBox::Slot()
+		.Padding(5.f, 0.f)
+		.AutoWidth()
+		.HAlign(HAlign_Fill)
+		[
+			SNew(SButton)
+			.HAlign(HAlign_Fill)
+			.Visibility(this, &FPropertyAnimatorCoreEditorSequencerTimeSourceChannelTypeCustomization::GetCreateTrackButtonVisibility)
+			.OnClicked(this, &FPropertyAnimatorCoreEditorSequencerTimeSourceChannelTypeCustomization::OnCreateTrackButtonClicked)
+			.IsEnabled(this, &FPropertyAnimatorCoreEditorSequencerTimeSourceChannelTypeCustomization::IsCreateTrackButtonEnabled)
+			[
+				SNew(STextBlock)
+				.Font(IPropertyTypeCustomizationUtils::GetRegularFont())
+				.Text(LOCTEXT("AddSequencerTrack", "Create track"))
+				.ToolTipText(LOCTEXT("AddSequencerTrackTooltip", "Create a sequencer track linked to this time source"))
+			]
+		]
+	];
+}
+
+void FPropertyAnimatorCoreEditorSequencerTimeSourceChannelTypeCustomization::CustomizeChildren(TSharedRef<IPropertyHandle> InPropertyHandle, IDetailChildrenBuilder& InBuilder, IPropertyTypeCustomizationUtils& InUtils)
+{
+	// No implementation
+}
+
+FReply FPropertyAnimatorCoreEditorSequencerTimeSourceChannelTypeCustomization::OnCreateTrackButtonClicked()
+{
+	const TArray<UObject*> BindingObjects = GetBindingObjects();
+
+	if (BindingObjects.IsEmpty())
+	{
+		return FReply::Handled();
+	}
+
+	TArray<AActor*> SelectedActors = GetSelectedActors();
+
+	FMovieSceneAnimatorTrackEditor::OnAddAnimatorTrack.Broadcast(BindingObjects);
+
+	// Reselect actors after track was created
+	if (GEditor)
+	{
+		GEditor->SelectNone(/** Notify */false, /** DeselectBSP */true);
+
+		for (int32 Index = 0; Index < SelectedActors.Num(); Index++)
+		{
+			GEditor->SelectActor(SelectedActors[Index], /** Selected */true, /** Notify */Index == SelectedActors.Num() - 1);
+		}
+	}
+
+	return FReply::Handled();
+}
+
+EVisibility FPropertyAnimatorCoreEditorSequencerTimeSourceChannelTypeCustomization::GetCreateTrackButtonVisibility() const
+{
+	return EvalTimePropertyHandle.IsValid() && EvalTimePropertyHandle->IsValidHandle()
+		? EVisibility::Visible
+		: EVisibility::Collapsed;
+}
+
+bool FPropertyAnimatorCoreEditorSequencerTimeSourceChannelTypeCustomization::IsCreateTrackButtonEnabled() const
+{
+	const TArray<UObject*> BindingObjects = GetBindingObjects();
+
+	int32 Count = 0;
+	FMovieSceneAnimatorTrackEditor::OnGetAnimatorTrackCount.Broadcast(BindingObjects, Count);
+
+	return !BindingObjects.IsEmpty() && BindingObjects.Num() != Count;
+}
+
+TArray<AActor*> FPropertyAnimatorCoreEditorSequencerTimeSourceChannelTypeCustomization::GetSelectedActors() const
+{
+	TArray<AActor*> SelectedActors;
+
+	if (GEditor)
+	{
+		if (USelection* ActorSelection = GEditor->GetSelectedActors())
+		{
+			ActorSelection->GetSelectedObjects<AActor>(SelectedActors);
+		}
+	}
+
+	return SelectedActors;
+}
+
+TArray<UObject*> FPropertyAnimatorCoreEditorSequencerTimeSourceChannelTypeCustomization::GetBindingObjects() const
+{
+	TArray<UObject*> BindingObjects;
+
+	if (!EvalTimePropertyHandle.IsValid() || !EvalTimePropertyHandle->IsValidHandle())
+	{
+		return BindingObjects;
+	}
+
+	EvalTimePropertyHandle->GetOuterObjects(BindingObjects);
+
+	return BindingObjects;
+}
+
+#undef LOCTEXT_NAMESPACE
diff --git a/Engine/Plugins/Experimental/PropertyAnimatorCore/Source/PropertyAnimatorCoreEditor/Private/Customizations/PropertyAnimatorCoreEditorSequencerTimeSourceEvalResultTypeCustomization.h b/Engine/Plugins/Experimental/PropertyAnimatorCore/Source/PropertyAnimatorCoreEditor/Private/Customizations/PropertyAnimatorCoreEditorSequencerTimeSourceEvalResultTypeCustomization.h
new file mode 100644
--- /dev/null
+++ b/Engine/Plugins/Experimental/PropertyAnimatorCore/Source/PropertyAnimatorCoreEditor/Private/Customizations/PropertyAnimatorCoreEditorSequencerTimeSourceEvalResultTypeCustomization.h
@@ -0,0 +1,37 @@
+// Copyright Epic Games, Inc. All Rights Reserved.
+
+#pragma once
+
+#include "Containers/Array.h"
+#include "IPropertyTypeCustomization.h"
+#include "Templates/SharedPointer.h"
+
+class AActor;
+class FReply;
+class IPropertyHandle;
+class UObject;
+struct EVisibility;
+
+/** Type customization for FPropertyAnimatorCoreSequencerTimeSourceEvalResult */
+class FPropertyAnimatorCoreEditorSequencerTimeSourceChannelTypeCustomization : public IPropertyTypeCustomization
+{
+public:
+	static TSharedRef<IPropertyTypeCustomization> MakeInstance()
+	{
+		return MakeShared<FPropertyAnimatorCoreEditorSequencerTimeSourceChannelTypeCustomization>();
+	}
+
+	//~ Begin IPropertyTypeCustomization
+	virtual void CustomizeHeader(TSharedRef<IPropertyHandle> InPropertyHandle, FDetailWidgetRow& InRow, IPropertyTypeCustomizationUtils& InUtils) override;
+	virtual void CustomizeChildren(TSharedRef<IPropertyHandle> InPropertyHandle, IDetailChildrenBuilder& InBuilder, IPropertyTypeCustomizationUtils& InUtils) override;
+	//~ End IPropertyTypeCustomization
+
+protected:
+	FReply OnCreateTrackButtonClicked();
+	EVisibility GetCreateTrackButtonVisibility() const;
+	bool IsCreateTrackButtonEnabled() const;
+	TArray<AActor*> GetSelectedActors() const;
+	TArray<UObject*> GetBindingObjects() const;
+
+	TSharedPtr<IPropertyHandle> EvalTimePropertyHandle;
+};
diff --git a/Engine/Plugins/Experimental/PropertyAnimatorCore/Source/PropertyAnimatorCoreEditor/Private/PropertyAnimatorCoreEditorModule.cpp b/Engine/Plugins/Experimental/PropertyAnimatorCore/Source/PropertyAnimatorCoreEditor/Private/PropertyAnimatorCoreEditorModule.cpp
--- a/Engine/Plugins/Experimental/PropertyAnimatorCore/Source/PropertyAnimatorCoreEditor/Private/PropertyAnimatorCoreEditorModule.cpp
+++ b/Engine/Plugins/Experimental/PropertyAnimatorCore/Source/PropertyAnimatorCoreEditor/Private/PropertyAnimatorCoreEditorModule.cpp
@@ -6,7 +6,7 @@
 #include "Customizations/PropertyAnimatorCoreEditorContextTypeCustomization.h"
 #include "Customizations/PropertyAnimatorCoreEditorDetailCustomization.h"
 #include "Customizations/PropertyAnimatorCoreEditorManualStateTypeCustomization.h"
-#include "Customizations/PropertyAnimatorCoreEditorSequencerTimeSourceChannelTypeCustomization.h"
+#include "Customizations/PropertyAnimatorCoreEditorSequencerTimeSourceEvalResultTypeCustomization.h"
 #include "ISequencerModule.h"
 #include "Modules/ModuleManager.h"
 #include "Properties/PropertyAnimatorCoreContext.h"
@@ -18,9 +18,11 @@
 void FPropertyAnimatorCoreEditorModule::StartupModule()
 {
 	FPropertyEditorModule& PropertyModule = FModuleManager::LoadModuleChecked<FPropertyEditorModule>("PropertyEditor");
+
 	PropertyModule.RegisterCustomPropertyTypeLayout(UPropertyAnimatorCoreContext::StaticClass()->GetFName(), FOnGetPropertyTypeCustomizationInstance::CreateStatic(&FPropertyAnimatorCoreEditorContextTypeCustomization::MakeInstance));
-	PropertyModule.RegisterCustomPropertyTypeLayout(FPropertyAnimatorCoreSequencerTimeSourceChannel::StaticStruct()->GetFName(), FOnGetPropertyTypeCustomizationInstance::CreateStatic(&FPropertyAnimatorCoreEditorSequencerTimeSourceChannelTypeCustomization::MakeInstance));
 	PropertyModule.RegisterCustomPropertyTypeLayout(FPropertyAnimatorCoreManualState::StaticStruct()->GetFName(), FOnGetPropertyTypeCustomizationInstance::CreateStatic(&FPropertyAnimatorCoreEditorManualStateTypeCustomization::MakeInstance));
+	PropertyModule.RegisterCustomPropertyTypeLayout(FPropertyAnimatorCoreSequencerTimeSourceEvalResult::StaticStruct()->GetFName(), FOnGetPropertyTypeCustomizationInstance::CreateStatic(&FPropertyAnimatorCoreEditorSequencerTimeSourceChannelTypeCustomization::MakeInstance));
+
 	PropertyModule.RegisterCustomClassLayout(UPropertyAnimatorCoreBase::StaticClass()->GetFName(), FOnGetDetailCustomizationInstance::CreateStatic(&FPropertyAnimatorCoreEditorDetailCustomization::MakeInstance));
 
 	ISequencerModule& SequencerModule = FModuleManager::LoadModuleChecked<ISequencerModule>("Sequencer");
@@ -32,9 +34,11 @@ void FPropertyAnimatorCoreEditorModule::ShutdownModule()
 	if (FModuleManager::Get().IsModuleLoaded("PropertyEditor"))
 	{
 		FPropertyEditorModule& PropertyModule = FModuleManager::GetModuleChecked<FPropertyEditorModule>("PropertyEditor");
+
 		PropertyModule.UnregisterCustomPropertyTypeLayout(UPropertyAnimatorCoreContext::StaticClass()->GetFName());
-		PropertyModule.UnregisterCustomPropertyTypeLayout(FPropertyAnimatorCoreSequencerTimeSourceChannel::StaticStruct()->GetFName());
 		PropertyModule.UnregisterCustomPropertyTypeLayout(FPropertyAnimatorCoreManualState::StaticStruct()->GetFName());
+		PropertyModule.UnregisterCustomPropertyTypeLayout(FPropertyAnimatorCoreSequencerTimeSourceEvalResult::StaticStruct()->GetFName());
+
 		PropertyModule.UnregisterCustomClassLayout(UPropertyAnimatorCoreBase::StaticClass()->GetFName());
 	}
 
diff --git a/Engine/Plugins/Experimental/PropertyAnimatorCore/Source/PropertyAnimatorCoreEditor/Private/Sequencer/MovieSceneAnimatorTrackEditor.cpp b/Engine/Plugins/Experimental/PropertyAnimatorCore/Source/PropertyAnimatorCoreEditor/Private/Sequencer/MovieSceneAnimatorTrackEditor.cpp
--- a/Engine/Plugins/Experimental/PropertyAnimatorCore/Source/PropertyAnimatorCoreEditor/Private/Sequencer/MovieSceneAnimatorTrackEditor.cpp
+++ b/Engine/Plugins/Experimental/PropertyAnimatorCore/Source/PropertyAnimatorCoreEditor/Private/Sequencer/MovieSceneAnimatorTrackEditor.cpp
@@ -2,9 +2,9 @@
 
 #include "MovieSceneAnimatorTrackEditor.h"
 
+#include "Components/PropertyAnimatorCoreComponent.h"
 #include "LevelSequence.h"
 #include "Sequencer/MovieSceneAnimatorSection.h"
-#include "Styling/SlateIconFinder.h"
 
 #define LOCTEXT_NAMESPACE "MovieSceneAnimatorTrackEditor"
 
@@ -19,17 +19,7 @@ FMovieSceneAnimatorTrackEditor::~FMovieSceneAnimatorTrackEditor()
 
 void FMovieSceneAnimatorTrackEditor::BuildAddTrackMenu(FMenuBuilder& InMenuBuilder)
 {
-	constexpr uint8 Channel = 0;
-
-	InMenuBuilder.AddMenuEntry(
-		LOCTEXT("AddAnimatorTrack.Label", "Animator"),
-		LOCTEXT("AddAnimatorTrack.Tooltip", "Adds a new track that uses the time of the current sequence to drive animators."),
-		FSlateIconFinder::FindIconForClass(UMovieSceneAnimatorTrack::StaticClass()),
-		FUIAction(
-			FExecuteAction::CreateSP(this, &FMovieSceneAnimatorTrackEditor::ExecuteAddTrack, Channel),
-			FCanExecuteAction::CreateSP(this, &FMovieSceneAnimatorTrackEditor::CanExecuteAddTrack)
-		)
-	);
+	// Empty, do not allow the creation of new tracks, only allow through sequencer time source
 }
 
 TSharedPtr<SWidget> FMovieSceneAnimatorTrackEditor::BuildOutlinerEditWidget(const FGuid& InObjectBinding, UMovieSceneTrack* InTrack, const FBuildEditWidgetParams& InParams)
@@ -56,7 +46,7 @@ void FMovieSceneAnimatorTrackEditor::BindDelegates()
 	OnGetAnimatorTrackCount.AddSP(this, &FMovieSceneAnimatorTrackEditor::GetTrackCount);
 }
 
-void FMovieSceneAnimatorTrackEditor::GetTrackCount(uint8 InChannel, int32& OutCount) const
+void FMovieSceneAnimatorTrackEditor::GetTrackCount(const TArray<UObject*>& InOwners, int32& OutCount) const
 {
 	const UMovieScene* FocusedMovieScene = GetFocusedMovieScene();
 	if (!FocusedMovieScene)
@@ -64,11 +54,32 @@ void FMovieSceneAnimatorTrackEditor::GetTrackCount(uint8 InChannel, int32& OutCo
 		return;
 	}
 
-	for (UMovieSceneTrack* Track : FocusedMovieScene->GetTracks())
+	const TSharedPtr<ISequencer> SequencerPtr = GetSequencer();
+	if (!SequencerPtr)
+	{
+		return;
+	}
+
+	for (UObject* Owner : InOwners)
 	{
-		if (const UMovieSceneAnimatorTrack* AnimatorTrack = Cast<UMovieSceneAnimatorTrack>(Track))
+		if (!Owner)
+		{
+			continue;
+		}
+
+		const FGuid ObjectBinding = SequencerPtr->GetHandleToObject(Owner, /** Create */false);
+
+		if (!ObjectBinding.IsValid())
+		{
+			return;
+		}
+
+		for (const UMovieSceneTrack* Track : FocusedMovieScene->FindTracks(UMovieSceneAnimatorTrack::StaticClass(), ObjectBinding))
 		{
-			OutCount += AnimatorTrack->GetChannelCount(InChannel);
+			if (Track->IsA<UMovieSceneAnimatorTrack>())
+			{
+				OutCount++;
+			}
 		}
 	}
 }
@@ -78,7 +89,7 @@ bool FMovieSceneAnimatorTrackEditor::CanExecuteAddTrack() const
 	return GetSequencer() && GetFocusedMovieScene();
 }
 
-void FMovieSceneAnimatorTrackEditor::ExecuteAddTrack(uint8 InChannel)
+void FMovieSceneAnimatorTrackEditor::ExecuteAddTrack(const TArray<UObject*>& InOwners)
 {
 	UMovieScene* FocusedMovieScene = GetFocusedMovieScene();
 	if (!FocusedMovieScene || FocusedMovieScene->IsReadOnly())
@@ -96,13 +107,40 @@ void FMovieSceneAnimatorTrackEditor::ExecuteAddTrack(uint8 InChannel)
 
 	FocusedMovieScene->Modify();
 
-	UMovieSceneAnimatorTrack* NewTrack = NewObject<UMovieSceneAnimatorTrack>(FocusedMovieScene, NAME_None, RF_Transactional);
-	UMovieSceneAnimatorSection* NewSection = Cast<UMovieSceneAnimatorSection>(NewTrack->CreateNewSection());
-	NewSection->SetChannel(InChannel);
-	NewTrack->AddSection(*NewSection);
+	for (UObject* Owner : InOwners)
+	{
+		if (!Owner || !Owner->GetOuter())
+		{
+			continue;
+		}
+
+		const FGuid ObjectBinding = SequencerPtr->GetHandleToObject(Owner, /** Create */true);
+
+		if (FocusedMovieScene->FindSpawnable(ObjectBinding))
+		{
+			// we only want to add tracks for possessables
+			return;
+		}
+
+		if (!FocusedMovieScene->FindTrack<UMovieSceneAnimatorTrack>(ObjectBinding))
+		{
+			UMovieSceneAnimatorTrack* NewTrack = FocusedMovieScene->AddTrack<UMovieSceneAnimatorTrack>(ObjectBinding);
+
+			if (const UPropertyAnimatorCoreBase* Animator = Cast<UPropertyAnimatorCoreBase>(Owner->GetOuter()))
+			{
+				NewTrack->SetDisplayName(FText::Format(LOCTEXT("MovieSceneAnimatorTrackName", "Animator {0} Track"), FText::FromName(Animator->GetAnimatorOriginalName())));
+			}
+			else
+			{
+				NewTrack->SetDisplayName(LOCTEXT("MovieSceneAnimatorComponentTrackName", "Animator Component Track"));
+			}
+
+			UMovieSceneAnimatorSection* NewSection = Cast<UMovieSceneAnimatorSection>(NewTrack->CreateNewSection());
+			NewTrack->AddSection(*NewSection);
+		}
+	}
 
-	FocusedMovieScene->AddGivenTrack(NewTrack);
-	SequencerPtr->OnAddTrack(NewTrack, FGuid());
+	SequencerPtr->NotifyMovieSceneDataChanged(EMovieSceneDataChangeType::MovieSceneStructureItemAdded);
 }
 
 #undef LOCTEXT_NAMESPACE
diff --git a/Engine/Plugins/Experimental/PropertyAnimatorCore/Source/PropertyAnimatorCoreEditor/Private/Sequencer/MovieSceneAnimatorTrackEditor.h b/Engine/Plugins/Experimental/PropertyAnimatorCore/Source/PropertyAnimatorCoreEditor/Private/Sequencer/MovieSceneAnimatorTrackEditor.h
--- a/Engine/Plugins/Experimental/PropertyAnimatorCore/Source/PropertyAnimatorCoreEditor/Private/Sequencer/MovieSceneAnimatorTrackEditor.h
+++ b/Engine/Plugins/Experimental/PropertyAnimatorCore/Source/PropertyAnimatorCoreEditor/Private/Sequencer/MovieSceneAnimatorTrackEditor.h
@@ -11,10 +11,10 @@ class UMovieSceneSequence;
 class FMovieSceneAnimatorTrackEditor : public FKeyframeTrackEditor<UMovieSceneAnimatorTrack>
 {
 public:
-	DECLARE_MULTICAST_DELEGATE_OneParam(FOnAddAnimatorTrack, uint8 /** Channel */)
+	DECLARE_MULTICAST_DELEGATE_OneParam(FOnAddAnimatorTrack, const TArray<UObject*>& /** InOwners */)
 	static FOnAddAnimatorTrack OnAddAnimatorTrack;
 
-	DECLARE_MULTICAST_DELEGATE_TwoParams(FOnGetAnimatorTrackCount, uint8 /** Channel */, int32& /** Count */)
+	DECLARE_MULTICAST_DELEGATE_TwoParams(FOnGetAnimatorTrackCount, const TArray<UObject*>& /** InOwner */, int32& /** Count */)
 	static FOnGetAnimatorTrackCount OnGetAnimatorTrackCount;
 
 	FMovieSceneAnimatorTrackEditor(const TSharedRef<ISequencer>& InSequencer)
@@ -39,8 +39,8 @@ private:
 
 	void BindDelegates();
 
-	void GetTrackCount(uint8 InChannel, int32& OutCount) const;
+	void GetTrackCount(const TArray<UObject*>& InOwners, int32& OutCount) const;
 
 	bool CanExecuteAddTrack() const;
-	void ExecuteAddTrack(uint8 InChannel);
+	void ExecuteAddTrack(const TArray<UObject*>& InOwners);
 };
