MotionDesign : PropertyAnimatorCore
- Adding pre and post roll support in animator track to evaluate first and last time value instead of restoring state when out of section range

diff --git a/Engine/Plugins/Experimental/PropertyAnimatorCore/Source/PropertyAnimatorCore/Private/Sequencer/MovieSceneAnimatorExecutionToken.h b/Engine/Plugins/Experimental/PropertyAnimatorCore/Source/PropertyAnimatorCore/Private/Sequencer/MovieSceneAnimatorExecutionToken.h
--- a/Engine/Plugins/Experimental/PropertyAnimatorCore/Source/PropertyAnimatorCore/Private/Sequencer/MovieSceneAnimatorExecutionToken.h
+++ b/Engine/Plugins/Experimental/PropertyAnimatorCore/Source/PropertyAnimatorCore/Private/Sequencer/MovieSceneAnimatorExecutionToken.h
@@ -54,21 +54,75 @@ struct FMovieSceneAnimatorExecutionToken : IMovieSceneExecutionToken
 	{
 		if (InOperand.ObjectBindingID.IsValid() && SectionData.Section)
 		{
-			double EvaluatedTime = InContext.GetTime().AsDecimal();
+			double EvaluatedTime = 0;
 
-			if (SectionData.EvalTimeMode == EMovieSceneAnimatorEvalTimeMode::Section)
+			if (InContext.IsPreRoll())
 			{
-				EvaluatedTime -= SectionData.Section->GetInclusiveStartFrame().Value;
+				switch (SectionData.EvalTimeMode)
+				{
+					case EMovieSceneAnimatorEvalTimeMode::Sequence:
+					{
+						EvaluatedTime = (SectionData.Section->GetInclusiveStartFrame().Value * 1.0) / InContext.GetFrameRate().AsDecimal();
+					}
+					break;
+
+					case EMovieSceneAnimatorEvalTimeMode::Section:
+					{
+						EvaluatedTime = 0;
+					}
+					break;
+
+					case EMovieSceneAnimatorEvalTimeMode::Custom:
+					{
+						EvaluatedTime = SectionData.CustomStartTime;
+					}
+					break;
+				}
 			}
+			else if (InContext.IsPostRoll())
+			{
+				switch (SectionData.EvalTimeMode)
+				{
+					case EMovieSceneAnimatorEvalTimeMode::Sequence:
+					{
+						EvaluatedTime = (SectionData.Section->GetExclusiveEndFrame().Value * 1.0) / InContext.GetFrameRate().AsDecimal();
+					}
+					break;
+
+					case EMovieSceneAnimatorEvalTimeMode::Section:
+					{
+						const double SectionStartTime = (SectionData.Section->GetInclusiveStartFrame().Value * 1.0) / InContext.GetFrameRate().AsDecimal();
+						const double SectionEndTime = (SectionData.Section->GetExclusiveEndFrame().Value * 1.0) / InContext.GetFrameRate().AsDecimal();
+
+						EvaluatedTime = SectionEndTime - SectionStartTime;
+					}
+					break;
+
+					case EMovieSceneAnimatorEvalTimeMode::Custom:
+					{
+						EvaluatedTime = SectionData.CustomEndTime;
+					}
+					break;
+				}
+			}
+			else
+			{
+				EvaluatedTime = InContext.GetTime().AsDecimal();
 
-			EvaluatedTime /= InContext.GetFrameRate().AsDecimal();
+				if (SectionData.EvalTimeMode == EMovieSceneAnimatorEvalTimeMode::Section)
+				{
+					EvaluatedTime -= SectionData.Section->GetInclusiveStartFrame().Value;
+				}
 
-			if (SectionData.EvalTimeMode == EMovieSceneAnimatorEvalTimeMode::Custom)
-			{
-				const double SectionStartTime = (SectionData.Section->GetInclusiveStartFrame().Value * 1.f) / InContext.GetFrameRate().AsDecimal();
-				const double SectionEndTime = (SectionData.Section->GetExclusiveEndFrame().Value * 1.f) / InContext.GetFrameRate().AsDecimal();
+				EvaluatedTime /= InContext.GetFrameRate().AsDecimal();
 
-				EvaluatedTime = FMath::GetMappedRangeValueClamped(FVector2D(SectionStartTime, SectionEndTime), FVector2D(SectionData.CustomStartTime, SectionData.CustomEndTime), EvaluatedTime);
+				if (SectionData.EvalTimeMode == EMovieSceneAnimatorEvalTimeMode::Custom)
+				{
+					const double SectionStartTime = (SectionData.Section->GetInclusiveStartFrame().Value * 1.0) / InContext.GetFrameRate().AsDecimal();
+					const double SectionEndTime = (SectionData.Section->GetExclusiveEndFrame().Value * 1.0) / InContext.GetFrameRate().AsDecimal();
+
+					EvaluatedTime = FMath::GetMappedRangeValueClamped(FVector2D(SectionStartTime, SectionEndTime), FVector2D(SectionData.CustomStartTime, SectionData.CustomEndTime), EvaluatedTime);
+				}
 			}
 
 			const float Magnitude = SectionData.Section->EvaluateEasing(InContext.GetTime());
