MotionDesign : PropertyAnimatorCore
- Property locator path to identify an owner within an actor and the property within
- Adding preset support for linked properties and animator, property preset can be reused on multiple animator through the link property menu, animator preset can be used in the add animator menu
- Presets are saved as uasset to easily discover and share them, they use json format to (de)serialize data

diff --git a/Engine/Plugins/Experimental/PropertyAnimator/Source/PropertyAnimator/Private/Animators/PropertyAnimatorClock.cpp b/Engine/Plugins/Experimental/PropertyAnimator/Source/PropertyAnimator/Private/Animators/PropertyAnimatorClock.cpp
--- a/Engine/Plugins/Experimental/PropertyAnimator/Source/PropertyAnimator/Private/Animators/PropertyAnimatorClock.cpp
+++ b/Engine/Plugins/Experimental/PropertyAnimator/Source/PropertyAnimator/Private/Animators/PropertyAnimatorClock.cpp
@@ -2,6 +2,8 @@
 
 #include "Animators/PropertyAnimatorClock.h"
 
+#include "Dom/JsonObject.h"
+#include "Dom/JsonValue.h"
 #include "Misc/DateTime.h"
 #include "Properties/Converters/PropertyAnimatorCoreConverterBase.h"
 #include "Properties/Handlers/PropertyAnimatorCoreHandlerBase.h"
@@ -114,4 +116,34 @@ void UPropertyAnimatorClock::EvaluateProperties(FInstancedPropertyBag& InParamet
 
 		return true;
 	});
-}
\ No newline at end of file
+}
+
+bool UPropertyAnimatorClock::ImportPreset(const UPropertyAnimatorCorePresetBase* InPreset, const TSharedRef<FJsonValue>& InValue)
+{
+	const TSharedPtr<FJsonObject>* JsonAnimatorObject;
+
+	if (Super::ImportPreset(InPreset, InValue) && InValue->TryGetObject(JsonAnimatorObject))
+	{
+		FString JsonDisplayFormat = DisplayFormat;
+		(*JsonAnimatorObject)->TryGetStringField(GET_MEMBER_NAME_STRING_CHECKED(UPropertyAnimatorClock, DisplayFormat), JsonDisplayFormat);
+		SetDisplayFormat(JsonDisplayFormat);
+
+		return true;
+	}
+
+	return false;
+}
+
+bool UPropertyAnimatorClock::ExportPreset(const UPropertyAnimatorCorePresetBase* InPreset, TSharedPtr<FJsonValue>& OutValue)
+{
+	const TSharedPtr<FJsonObject>* JsonAnimatorObject;
+
+	if (Super::ExportPreset(InPreset, OutValue) && OutValue->TryGetObject(JsonAnimatorObject))
+	{
+		(*JsonAnimatorObject)->SetStringField(GET_MEMBER_NAME_STRING_CHECKED(UPropertyAnimatorClock, DisplayFormat), DisplayFormat);
+
+		return true;
+	}
+
+	return false;
+}
diff --git a/Engine/Plugins/Experimental/PropertyAnimator/Source/PropertyAnimator/Private/Animators/PropertyAnimatorCounter.cpp b/Engine/Plugins/Experimental/PropertyAnimator/Source/PropertyAnimator/Private/Animators/PropertyAnimatorCounter.cpp
--- a/Engine/Plugins/Experimental/PropertyAnimator/Source/PropertyAnimator/Private/Animators/PropertyAnimatorCounter.cpp
+++ b/Engine/Plugins/Experimental/PropertyAnimator/Source/PropertyAnimator/Private/Animators/PropertyAnimatorCounter.cpp
@@ -2,6 +2,8 @@
 
 #include "Animators/PropertyAnimatorCounter.h"
 
+#include "Dom/JsonObject.h"
+#include "Dom/JsonValue.h"
 #include "Properties/Converters/PropertyAnimatorCoreConverterBase.h"
 #include "Settings/PropertyAnimatorSettings.h"
 #include "Subsystems/PropertyAnimatorCoreSubsystem.h"
@@ -182,25 +184,27 @@ FString UPropertyAnimatorCounter::FormatNumber(double InNumber) const
 {
 	FText Output = FText::GetEmpty();
 
-	if (!bUseCustomFormat)
+	if (const FPropertyAnimatorCounterFormat* Format = GetFormat())
 	{
-		if (const UPropertyAnimatorSettings* AnimatorSettings = GetDefault<UPropertyAnimatorSettings>())
-		{
-			if (const FPropertyAnimatorCounterFormat* Format = AnimatorSettings->GetCounterFormat(PresetFormatName))
-			{
-				Output = FText::Format(DisplayPattern, FText::FromString(Format->FormatNumber(InNumber)));
-			}
-		}
+		Output = FText::Format(DisplayPattern, FText::FromString(Format->FormatNumber(InNumber)));
 	}
-	else
+
+	return Output.ToString();
+}
+
+const FPropertyAnimatorCounterFormat* UPropertyAnimatorCounter::GetFormat() const
+{
+	if (bUseCustomFormat)
 	{
-		if (const FPropertyAnimatorCounterFormat* Format = CustomFormat.GetPtr<FPropertyAnimatorCounterFormat>())
-		{
-			Output = FText::Format(DisplayPattern, FText::FromString(Format->FormatNumber(InNumber)));
-		}
+		return CustomFormat.GetPtr<FPropertyAnimatorCounterFormat>();
 	}
 
-	return Output.ToString();
+	if (const UPropertyAnimatorSettings* AnimatorSettings = GetDefault<UPropertyAnimatorSettings>())
+	{
+		return AnimatorSettings->GetCounterFormat(PresetFormatName);
+	}
+
+	return nullptr;
 }
 
 #if WITH_EDITOR
@@ -279,6 +283,78 @@ void UPropertyAnimatorCounter::EvaluateProperties(FInstancedPropertyBag& InParam
 	});
 }
 
+bool UPropertyAnimatorCounter::ImportPreset(const UPropertyAnimatorCorePresetBase* InPreset, const TSharedRef<FJsonValue>& InValue)
+{
+	const TSharedPtr<FJsonObject>* JsonAnimatorObject;
+
+	if (Super::ImportPreset(InPreset, InValue) && InValue->TryGetObject(JsonAnimatorObject))
+	{
+		FString JsonDisplayPattern = DisplayPattern.ToString();
+		(*JsonAnimatorObject)->TryGetStringField(GET_MEMBER_NAME_STRING_CHECKED(UPropertyAnimatorCounter, DisplayPattern), JsonDisplayPattern);
+		SetDisplayPattern(FText::FromString(JsonDisplayPattern));
+
+		FString JsonPresetFormatNameStr = PresetFormatName.ToString();
+		(*JsonAnimatorObject)->TryGetStringField(GET_MEMBER_NAME_STRING_CHECKED(UPropertyAnimatorCounter, PresetFormatName), JsonPresetFormatNameStr);
+		FName JsonPresetFormatName(JsonPresetFormatNameStr);
+		SetPresetFormatName(JsonPresetFormatName);
+
+		if (!PresetFormatName.IsEqual(JsonPresetFormatName))
+		{
+			FPropertyAnimatorCounterFormat Format;
+			Format.FormatName = JsonPresetFormatName;
+
+			(*JsonAnimatorObject)->TryGetNumberField(GET_MEMBER_NAME_STRING_CHECKED(FPropertyAnimatorCounterFormat, MinIntegerCount), Format.MinIntegerCount);
+			(*JsonAnimatorObject)->TryGetNumberField(GET_MEMBER_NAME_STRING_CHECKED(FPropertyAnimatorCounterFormat, MaxDecimalCount), Format.MaxDecimalCount);
+			(*JsonAnimatorObject)->TryGetNumberField(GET_MEMBER_NAME_STRING_CHECKED(FPropertyAnimatorCounterFormat, GroupingSize), Format.GroupingSize);
+			(*JsonAnimatorObject)->TryGetStringField(GET_MEMBER_NAME_STRING_CHECKED(FPropertyAnimatorCounterFormat, DecimalCharacter), Format.DecimalCharacter);
+			(*JsonAnimatorObject)->TryGetStringField(GET_MEMBER_NAME_STRING_CHECKED(FPropertyAnimatorCounterFormat, GroupingCharacter), Format.GroupingCharacter);
+			(*JsonAnimatorObject)->TryGetStringField(GET_MEMBER_NAME_STRING_CHECKED(FPropertyAnimatorCounterFormat, PaddingCharacter), Format.PaddingCharacter);
+
+			uint8 JsonRoundingMode = static_cast<uint8>(Format.RoundingMode);
+			(*JsonAnimatorObject)->TryGetNumberField(GET_MEMBER_NAME_STRING_CHECKED(FPropertyAnimatorCounterFormat, RoundingMode), JsonRoundingMode);
+			Format.RoundingMode = static_cast<EPropertyAnimatorCounterRoundingMode>(JsonRoundingMode);
+
+			(*JsonAnimatorObject)->TryGetBoolField(GET_MEMBER_NAME_STRING_CHECKED(FPropertyAnimatorCounterFormat, bUseSign), Format.bUseSign);
+			(*JsonAnimatorObject)->TryGetBoolField(GET_MEMBER_NAME_STRING_CHECKED(FPropertyAnimatorCounterFormat, bTruncate), Format.bTruncate);
+
+			SetUseCustomFormat(true);
+			SetCustomFormat(&Format);
+		}
+
+		return true;
+	}
+
+	return false;
+}
+
+bool UPropertyAnimatorCounter::ExportPreset(const UPropertyAnimatorCorePresetBase* InPreset, TSharedPtr<FJsonValue>& OutValue)
+{
+	const TSharedPtr<FJsonObject>* JsonAnimatorObject;
+
+	if (Super::ExportPreset(InPreset, OutValue) && OutValue->TryGetObject(JsonAnimatorObject))
+	{
+		(*JsonAnimatorObject)->SetStringField(GET_MEMBER_NAME_STRING_CHECKED(UPropertyAnimatorCounter, DisplayPattern), DisplayPattern.ToString());
+
+		if (const FPropertyAnimatorCounterFormat* Format = GetFormat())
+		{
+			(*JsonAnimatorObject)->SetStringField(GET_MEMBER_NAME_STRING_CHECKED(UPropertyAnimatorCounter, PresetFormatName), bUseCustomFormat ? Format->FormatName.ToString() : PresetFormatName.ToString());
+			(*JsonAnimatorObject)->SetNumberField(GET_MEMBER_NAME_STRING_CHECKED(FPropertyAnimatorCounterFormat, MinIntegerCount), Format->MinIntegerCount);
+			(*JsonAnimatorObject)->SetNumberField(GET_MEMBER_NAME_STRING_CHECKED(FPropertyAnimatorCounterFormat, MaxDecimalCount), Format->MaxDecimalCount);
+			(*JsonAnimatorObject)->SetNumberField(GET_MEMBER_NAME_STRING_CHECKED(FPropertyAnimatorCounterFormat, GroupingSize), Format->GroupingSize);
+			(*JsonAnimatorObject)->SetStringField(GET_MEMBER_NAME_STRING_CHECKED(FPropertyAnimatorCounterFormat, DecimalCharacter), Format->DecimalCharacter);
+			(*JsonAnimatorObject)->SetStringField(GET_MEMBER_NAME_STRING_CHECKED(FPropertyAnimatorCounterFormat, GroupingCharacter), Format->GroupingCharacter);
+			(*JsonAnimatorObject)->SetStringField(GET_MEMBER_NAME_STRING_CHECKED(FPropertyAnimatorCounterFormat, PaddingCharacter), Format->PaddingCharacter);
+			(*JsonAnimatorObject)->SetNumberField(GET_MEMBER_NAME_STRING_CHECKED(FPropertyAnimatorCounterFormat, RoundingMode), static_cast<uint8>(Format->RoundingMode));
+			(*JsonAnimatorObject)->SetBoolField(GET_MEMBER_NAME_STRING_CHECKED(FPropertyAnimatorCounterFormat, bUseSign), Format->bUseSign);
+			(*JsonAnimatorObject)->SetBoolField(GET_MEMBER_NAME_STRING_CHECKED(FPropertyAnimatorCounterFormat, bTruncate), Format->bTruncate);
+		}
+
+		return true;
+	}
+
+	return false;
+}
+
 TArray<FName> UPropertyAnimatorCounter::GetAvailableFormatNames() const
 {
 	TArray<FName> FormatNames;
diff --git a/Engine/Plugins/Experimental/PropertyAnimator/Source/PropertyAnimator/Private/Animators/PropertyAnimatorCurve.cpp b/Engine/Plugins/Experimental/PropertyAnimator/Source/PropertyAnimator/Private/Animators/PropertyAnimatorCurve.cpp
--- a/Engine/Plugins/Experimental/PropertyAnimator/Source/PropertyAnimator/Private/Animators/PropertyAnimatorCurve.cpp
+++ b/Engine/Plugins/Experimental/PropertyAnimator/Source/PropertyAnimator/Private/Animators/PropertyAnimatorCurve.cpp
@@ -4,6 +4,8 @@
 
 #include "Curves/PropertyAnimatorEaseCurve.h"
 #include "Curves/PropertyAnimatorWaveCurve.h"
+#include "Dom/JsonObject.h"
+#include "Dom/JsonValue.h"
 #include "Properties/PropertyAnimatorFloatContext.h"
 #include "UObject/ConstructorHelpers.h"
 
@@ -122,6 +124,110 @@ bool UPropertyAnimatorCurve::EvaluateProperty(const FPropertyAnimatorCoreData& I
 	return InContext->EvaluateProperty(InPropertyData, InParameters, OutEvaluationResult);
 }
 
+bool UPropertyAnimatorCurve::ImportPreset(const UPropertyAnimatorCorePresetBase* InPreset, const TSharedRef<FJsonValue>& InValue)
+{
+	const TSharedPtr<FJsonObject>* JsonAnimatorObject;
+
+	if (Super::ImportPreset(InPreset, InValue) && InValue->TryGetObject(JsonAnimatorObject))
+	{
+		FString JsonWaveCurve;
+		(*JsonAnimatorObject)->TryGetStringField(GET_MEMBER_NAME_STRING_CHECKED(UPropertyAnimatorCurve, WaveCurve), JsonWaveCurve);
+
+		if (UPropertyAnimatorWaveCurve* Curve = LoadObject<UPropertyAnimatorWaveCurve>(nullptr, *JsonWaveCurve))
+		{
+			SetWaveCurve(Curve);
+		}
+
+		bool bJsonEaseInEnabled = bEaseInEnabled;
+		(*JsonAnimatorObject)->TryGetBoolField(GET_MEMBER_NAME_STRING_CHECKED(UPropertyAnimatorCurve, bEaseInEnabled), bJsonEaseInEnabled);
+		SetEaseInEnabled(bJsonEaseInEnabled);
+
+		bool bJsonEaseOutEnabled = bEaseOutEnabled;
+		(*JsonAnimatorObject)->TryGetBoolField(GET_MEMBER_NAME_STRING_CHECKED(UPropertyAnimatorCurve, bEaseOutEnabled), bJsonEaseOutEnabled);
+		SetEaseOutEnabled(bJsonEaseOutEnabled);
+
+		const TSharedPtr<FJsonObject>* JsonEaseInObject;
+		(*JsonAnimatorObject)->TryGetObjectField(GET_MEMBER_NAME_STRING_CHECKED(UPropertyAnimatorCurve, EaseIn), JsonEaseInObject);
+
+		if (JsonEaseInObject)
+		{
+			FString JsonEaseCurve;
+			(*JsonEaseInObject)->TryGetStringField(GET_MEMBER_NAME_STRING_CHECKED(FPropertyAnimatorCurveEasing, EaseCurve), JsonEaseCurve);
+
+			double JsonEaseDuration;
+			(*JsonEaseInObject)->TryGetNumberField(GET_MEMBER_NAME_STRING_CHECKED(FPropertyAnimatorCurveEasing, EaseDuration), JsonEaseDuration);
+
+			if (UPropertyAnimatorEaseCurve* EaseCurve = LoadObject<UPropertyAnimatorEaseCurve>(nullptr, *JsonEaseCurve))
+			{
+				FPropertyAnimatorCurveEasing Easing;
+				Easing.EaseCurve = EaseCurve;
+				Easing.EaseDuration = JsonEaseDuration;
+				SetEaseIn(Easing);
+			}
+		}
+
+		const TSharedPtr<FJsonObject>* JsonEaseOutObject;
+		(*JsonAnimatorObject)->TryGetObjectField(GET_MEMBER_NAME_STRING_CHECKED(UPropertyAnimatorCurve, EaseOut), JsonEaseOutObject);
+
+		if (JsonEaseOutObject)
+		{
+			FString JsonEaseCurve;
+			(*JsonEaseOutObject)->TryGetStringField(GET_MEMBER_NAME_STRING_CHECKED(FPropertyAnimatorCurveEasing, EaseCurve), JsonEaseCurve);
+
+			double JsonEaseDuration;
+			(*JsonEaseOutObject)->TryGetNumberField(GET_MEMBER_NAME_STRING_CHECKED(FPropertyAnimatorCurveEasing, EaseDuration), JsonEaseDuration);
+
+			if (UPropertyAnimatorEaseCurve* EaseCurve = LoadObject<UPropertyAnimatorEaseCurve>(nullptr, *JsonEaseCurve))
+			{
+				FPropertyAnimatorCurveEasing Easing;
+				Easing.EaseCurve = EaseCurve;
+				Easing.EaseDuration = JsonEaseDuration;
+				SetEaseOut(Easing);
+			}
+		}
+
+		return true;
+	}
+
+	return false;
+}
+
+bool UPropertyAnimatorCurve::ExportPreset(const UPropertyAnimatorCorePresetBase* InPreset, TSharedPtr<FJsonValue>& OutValue)
+{
+	const TSharedPtr<FJsonObject>* JsonAnimatorObject;
+
+	if (Super::ExportPreset(InPreset, OutValue) && OutValue->TryGetObject(JsonAnimatorObject))
+	{
+		if (WaveCurve)
+		{
+			(*JsonAnimatorObject)->SetStringField(GET_MEMBER_NAME_STRING_CHECKED(UPropertyAnimatorCurve, WaveCurve), WaveCurve.GetPath());
+		}
+
+		(*JsonAnimatorObject)->SetBoolField(GET_MEMBER_NAME_STRING_CHECKED(UPropertyAnimatorCurve, bEaseInEnabled), bEaseInEnabled);
+		(*JsonAnimatorObject)->SetBoolField(GET_MEMBER_NAME_STRING_CHECKED(UPropertyAnimatorCurve, bEaseOutEnabled), bEaseOutEnabled);
+
+		if (EaseIn.EaseCurve)
+		{
+			TSharedRef<FJsonObject> JsonEaseObject = MakeShared<FJsonObject>();
+			JsonEaseObject->SetStringField(GET_MEMBER_NAME_STRING_CHECKED(FPropertyAnimatorCurveEasing, EaseCurve), EaseIn.EaseCurve.GetPath());
+			JsonEaseObject->SetNumberField(GET_MEMBER_NAME_STRING_CHECKED(FPropertyAnimatorCurveEasing, EaseDuration), EaseIn.EaseDuration);
+			(*JsonAnimatorObject)->SetObjectField(GET_MEMBER_NAME_STRING_CHECKED(UPropertyAnimatorCurve, EaseIn), JsonEaseObject);
+		}
+
+		if (EaseOut.EaseCurve)
+		{
+			TSharedRef<FJsonObject> JsonEaseObject = MakeShared<FJsonObject>();
+			JsonEaseObject->SetStringField(GET_MEMBER_NAME_STRING_CHECKED(FPropertyAnimatorCurveEasing, EaseCurve), EaseOut.EaseCurve.GetPath());
+			JsonEaseObject->SetNumberField(GET_MEMBER_NAME_STRING_CHECKED(FPropertyAnimatorCurveEasing, EaseDuration), EaseOut.EaseDuration);
+			(*JsonAnimatorObject)->SetObjectField(GET_MEMBER_NAME_STRING_CHECKED(UPropertyAnimatorCurve, EaseOut), JsonEaseObject);
+		}
+
+		return true;
+	}
+
+	return false;
+}
+
 void UPropertyAnimatorCurve::OnEaseInChanged()
 {
 	EaseIn.EaseDuration = FMath::Clamp(EaseIn.EaseDuration, 0, CycleDuration - EaseOut.EaseDuration);
diff --git a/Engine/Plugins/Experimental/PropertyAnimator/Source/PropertyAnimator/Private/Animators/PropertyAnimatorNumericBase.cpp b/Engine/Plugins/Experimental/PropertyAnimator/Source/PropertyAnimator/Private/Animators/PropertyAnimatorNumericBase.cpp
--- a/Engine/Plugins/Experimental/PropertyAnimator/Source/PropertyAnimator/Private/Animators/PropertyAnimatorNumericBase.cpp
+++ b/Engine/Plugins/Experimental/PropertyAnimator/Source/PropertyAnimator/Private/Animators/PropertyAnimatorNumericBase.cpp
@@ -2,6 +2,8 @@
 
 #include "Animators/PropertyAnimatorNumericBase.h"
 
+#include "Dom/JsonObject.h"
+#include "Dom/JsonValue.h"
 #include "Properties/PropertyAnimatorFloatContext.h"
 #include "Properties/PropertyAnimatorRotatorContext.h"
 #include "Properties/PropertyAnimatorVectorContext.h"
@@ -72,17 +74,6 @@ void UPropertyAnimatorNumericBase::SetCycleMode(EPropertyAnimatorCycleMode InMod
 	OnCycleModeChanged();
 }
 
-void UPropertyAnimatorNumericBase::SetTimeOffset(double InOffset)
-{
-	if (FMath::IsNearlyEqual(TimeOffset, InOffset))
-	{
-		return;
-	}
-
-	TimeOffset = InOffset;
-	OnTimeOffsetChanged();
-}
-
 void UPropertyAnimatorNumericBase::SetRandomTimeOffset(bool bInOffset)
 {
 	if (bRandomTimeOffset == bInOffset)
@@ -177,6 +168,7 @@ void UPropertyAnimatorNumericBase::EvaluateProperties(FInstancedPropertyBag& InP
 	{
 		const double RandomTimeOffset = bRandomTimeOffset ? RandomStream.GetFraction() : 0;
 
+		const double TimeOffset = InOptions->GetTimeOffset() / (InRangeMax + 1);
 		const double AbsTimeOffset = FMath::Abs(TimeOffset);
 		const double MaxTimeOffset = InRangeMax * AbsTimeOffset;
 		double PropertyTimeElapsed = TimeElapsed - MaxTimeOffset + RandomTimeOffset;
@@ -283,3 +275,66 @@ void UPropertyAnimatorNumericBase::OnAnimatorRegistered(FPropertyAnimatorCoreMet
 
 	InMetadata.Category = TEXT("Numeric");
 }
+
+bool UPropertyAnimatorNumericBase::ImportPreset(const UPropertyAnimatorCorePresetBase* InPreset, const TSharedRef<FJsonValue>& InValue)
+{
+	const TSharedPtr<FJsonObject>* JsonAnimatorObject;
+
+	if (Super::ImportPreset(InPreset, InValue) && InValue->TryGetObject(JsonAnimatorObject))
+	{
+		double JsonMagnitude = Magnitude;
+		(*JsonAnimatorObject)->TryGetNumberField(GET_MEMBER_NAME_STRING_CHECKED(UPropertyAnimatorNumericBase, Magnitude), JsonMagnitude);
+		SetMagnitude(JsonMagnitude);
+
+		if (CycleMode != EPropertyAnimatorCycleMode::None)
+		{
+			double JsonCycleMode = static_cast<double>(CycleMode);
+			(*JsonAnimatorObject)->TryGetNumberField(GET_MEMBER_NAME_STRING_CHECKED(UPropertyAnimatorNumericBase, CycleMode), JsonCycleMode);
+			SetCycleMode(static_cast<EPropertyAnimatorCycleMode>(JsonCycleMode));
+
+			double JsonCycleDuration = CycleDuration;
+			(*JsonAnimatorObject)->TryGetNumberField(GET_MEMBER_NAME_STRING_CHECKED(UPropertyAnimatorNumericBase, CycleDuration), JsonCycleDuration);
+			SetCycleDuration(JsonCycleDuration);
+
+			double JsonCycleGapDuration = CycleGapDuration;
+			(*JsonAnimatorObject)->TryGetNumberField(GET_MEMBER_NAME_STRING_CHECKED(UPropertyAnimatorNumericBase, CycleGapDuration), JsonCycleGapDuration);
+			SetCycleGapDuration(JsonCycleGapDuration);
+		}
+
+		bool bJsonRandomTimeOffset = bRandomTimeOffset;
+		(*JsonAnimatorObject)->TryGetBoolField(GET_MEMBER_NAME_STRING_CHECKED(UPropertyAnimatorNumericBase, bRandomTimeOffset), bJsonRandomTimeOffset);
+		SetRandomTimeOffset(bJsonRandomTimeOffset);
+
+		double JsonSeed = Seed;
+		(*JsonAnimatorObject)->TryGetNumberField(GET_MEMBER_NAME_STRING_CHECKED(UPropertyAnimatorNumericBase, Seed), JsonSeed);
+		SetSeed(JsonSeed);
+
+		return true;
+	}
+
+	return false;
+}
+
+bool UPropertyAnimatorNumericBase::ExportPreset(const UPropertyAnimatorCorePresetBase* InPreset, TSharedPtr<FJsonValue>& OutValue)
+{
+	const TSharedPtr<FJsonObject>* JsonAnimatorObject;
+
+	if (Super::ExportPreset(InPreset, OutValue) && OutValue->TryGetObject(JsonAnimatorObject))
+	{
+		(*JsonAnimatorObject)->SetNumberField(GET_MEMBER_NAME_STRING_CHECKED(UPropertyAnimatorNumericBase, Magnitude), Magnitude);
+
+		if (CycleMode != EPropertyAnimatorCycleMode::None)
+		{
+			(*JsonAnimatorObject)->SetNumberField(GET_MEMBER_NAME_STRING_CHECKED(UPropertyAnimatorNumericBase, CycleMode), static_cast<double>(CycleMode));
+			(*JsonAnimatorObject)->SetNumberField(GET_MEMBER_NAME_STRING_CHECKED(UPropertyAnimatorNumericBase, CycleDuration), CycleDuration);
+			(*JsonAnimatorObject)->SetNumberField(GET_MEMBER_NAME_STRING_CHECKED(UPropertyAnimatorNumericBase, CycleGapDuration), CycleGapDuration);
+		}
+
+		(*JsonAnimatorObject)->SetBoolField(GET_MEMBER_NAME_STRING_CHECKED(UPropertyAnimatorNumericBase, bRandomTimeOffset), bRandomTimeOffset);
+		(*JsonAnimatorObject)->SetNumberField(GET_MEMBER_NAME_STRING_CHECKED(UPropertyAnimatorNumericBase, Seed), Seed);
+
+		return true;
+	}
+
+	return false;
+}
diff --git a/Engine/Plugins/Experimental/PropertyAnimator/Source/PropertyAnimator/Private/Animators/PropertyAnimatorSoundWave.cpp b/Engine/Plugins/Experimental/PropertyAnimator/Source/PropertyAnimator/Private/Animators/PropertyAnimatorSoundWave.cpp
--- a/Engine/Plugins/Experimental/PropertyAnimator/Source/PropertyAnimator/Private/Animators/PropertyAnimatorSoundWave.cpp
+++ b/Engine/Plugins/Experimental/PropertyAnimator/Source/PropertyAnimator/Private/Animators/PropertyAnimatorSoundWave.cpp
@@ -2,6 +2,8 @@
 
 #include "Animators/PropertyAnimatorSoundWave.h"
 
+#include "Dom/JsonObject.h"
+#include "Dom/JsonValue.h"
 #include "LoudnessNRT.h"
 #include "Properties/Handlers/PropertyAnimatorCoreHandlerBase.h"
 #include "Properties/PropertyAnimatorFloatContext.h"
@@ -93,3 +95,46 @@ bool UPropertyAnimatorSoundWave::EvaluateProperty(const FPropertyAnimatorCoreDat
 
 	return false;
 }
+
+bool UPropertyAnimatorSoundWave::ImportPreset(const UPropertyAnimatorCorePresetBase* InPreset, const TSharedRef<FJsonValue>& InValue)
+{
+	const TSharedPtr<FJsonObject>* JsonAnimatorObject;
+
+	if (Super::ImportPreset(InPreset, InValue) && InValue->TryGetObject(JsonAnimatorObject))
+	{
+		FString JsonSoundWave;
+		(*JsonAnimatorObject)->TryGetStringField(GET_MEMBER_NAME_STRING_CHECKED(UPropertyAnimatorSoundWave, SampledSoundWave), JsonSoundWave);
+
+		if (USoundWave* SoundWave = LoadObject<USoundWave>(nullptr, *JsonSoundWave))
+		{
+			SetSampledSoundWave(SoundWave);
+		}
+
+		bool bJsonLoop = bLoop;
+		(*JsonAnimatorObject)->TryGetBoolField(GET_MEMBER_NAME_STRING_CHECKED(UPropertyAnimatorSoundWave, bLoop), bLoop);
+		SetLoop(bJsonLoop);
+
+		return true;
+	}
+
+	return false;
+}
+
+bool UPropertyAnimatorSoundWave::ExportPreset(const UPropertyAnimatorCorePresetBase* InPreset, TSharedPtr<FJsonValue>& OutValue)
+{
+	const TSharedPtr<FJsonObject>* JsonAnimatorObject;
+
+	if (Super::ExportPreset(InPreset, OutValue) && OutValue->TryGetObject(JsonAnimatorObject))
+	{
+		if (SampledSoundWave)
+		{
+			(*JsonAnimatorObject)->SetStringField(GET_MEMBER_NAME_STRING_CHECKED(UPropertyAnimatorSoundWave, SampledSoundWave), SampledSoundWave.GetPath());
+		}
+
+		(*JsonAnimatorObject)->SetBoolField(GET_MEMBER_NAME_STRING_CHECKED(UPropertyAnimatorSoundWave, bLoop), bLoop);
+
+		return true;
+	}
+
+	return false;
+}
diff --git a/Engine/Plugins/Experimental/PropertyAnimator/Source/PropertyAnimator/Private/Presets/PropertyAnimatorPresetLocation.h b/Engine/Plugins/Experimental/PropertyAnimator/Source/PropertyAnimator/Private/Presets/PropertyAnimatorPresetLocation.h
--- a/Engine/Plugins/Experimental/PropertyAnimator/Source/PropertyAnimator/Private/Presets/PropertyAnimatorPresetLocation.h
+++ b/Engine/Plugins/Experimental/PropertyAnimator/Source/PropertyAnimator/Private/Presets/PropertyAnimatorPresetLocation.h
@@ -2,7 +2,7 @@
 
 #pragma once
 
-#include "Presets/PropertyAnimatorCorePresetBase.h"
+#include "Presets/PropertyAnimatorCorePropertyPreset.h"
 #include "PropertyAnimatorPresetLocation.generated.h"
 
 class AActor;
@@ -12,18 +12,20 @@ class UPropertyAnimatorCoreBase;
  * Preset for position properties (X, Y, Z) on scene component
  */
 UCLASS(Transient)
-class UPropertyAnimatorPresetLocation : public UPropertyAnimatorCorePresetBase
+class UPropertyAnimatorPresetLocation : public UPropertyAnimatorCorePropertyPreset
 {
 	GENERATED_BODY()
 
 public:
 	UPropertyAnimatorPresetLocation()
-		: UPropertyAnimatorCorePresetBase(TEXT("Location"))
-	{}
+	{
+		PresetName = TEXT("Location");
+	}
 
 protected:
 	//~ Begin UPropertyAnimatorCorePresetBase
 	virtual void GetPresetProperties(const AActor* InActor, const UPropertyAnimatorCoreBase* InAnimator, TSet<FPropertyAnimatorCoreData>& OutProperties) const override;
 	virtual void OnPresetApplied(UPropertyAnimatorCoreBase* InAnimator, const TSet<FPropertyAnimatorCoreData>& InProperties) override;
+	virtual bool LoadPreset() override { return true; }
 	//~ End UPropertyAnimatorCorePresetBase
 };
\ No newline at end of file
diff --git a/Engine/Plugins/Experimental/PropertyAnimator/Source/PropertyAnimator/Private/Presets/PropertyAnimatorPresetRotation.h b/Engine/Plugins/Experimental/PropertyAnimator/Source/PropertyAnimator/Private/Presets/PropertyAnimatorPresetRotation.h
--- a/Engine/Plugins/Experimental/PropertyAnimator/Source/PropertyAnimator/Private/Presets/PropertyAnimatorPresetRotation.h
+++ b/Engine/Plugins/Experimental/PropertyAnimator/Source/PropertyAnimator/Private/Presets/PropertyAnimatorPresetRotation.h
@@ -2,7 +2,7 @@
 
 #pragma once
 
-#include "Presets/PropertyAnimatorCorePresetBase.h"
+#include "Presets/PropertyAnimatorCorePropertyPreset.h"
 #include "PropertyAnimatorPresetRotation.generated.h"
 
 class AActor;
@@ -12,18 +12,20 @@ class UPropertyAnimatorCoreBase;
  * Preset for position properties (Roll, Pitch, Yaw) on scene component
  */
 UCLASS(Transient)
-class UPropertyAnimatorPresetRotation : public UPropertyAnimatorCorePresetBase
+class UPropertyAnimatorPresetRotation : public UPropertyAnimatorCorePropertyPreset
 {
 	GENERATED_BODY()
 
 public:
 	UPropertyAnimatorPresetRotation()
-	: UPropertyAnimatorCorePresetBase(TEXT("Rotation"))
-	{}
+	{
+		PresetName = TEXT("Rotation");
+	}
 
 protected:
 	//~ Begin UPropertyAnimatorCorePresetBase
 	virtual void GetPresetProperties(const AActor* InActor, const UPropertyAnimatorCoreBase* InAnimator, TSet<FPropertyAnimatorCoreData>& OutProperties) const override;
 	virtual void OnPresetApplied(UPropertyAnimatorCoreBase* InAnimator, const TSet<FPropertyAnimatorCoreData>& InProperties) override;
+	virtual bool LoadPreset() override { return true; }
 	//~ End UPropertyAnimatorCorePresetBase
 };
\ No newline at end of file
diff --git a/Engine/Plugins/Experimental/PropertyAnimator/Source/PropertyAnimator/Private/Presets/PropertyAnimatorPresetScale.h b/Engine/Plugins/Experimental/PropertyAnimator/Source/PropertyAnimator/Private/Presets/PropertyAnimatorPresetScale.h
--- a/Engine/Plugins/Experimental/PropertyAnimator/Source/PropertyAnimator/Private/Presets/PropertyAnimatorPresetScale.h
+++ b/Engine/Plugins/Experimental/PropertyAnimator/Source/PropertyAnimator/Private/Presets/PropertyAnimatorPresetScale.h
@@ -2,7 +2,7 @@
 
 #pragma once
 
-#include "Presets/PropertyAnimatorCorePresetBase.h"
+#include "Presets/PropertyAnimatorCorePropertyPreset.h"
 #include "PropertyAnimatorPresetScale.generated.h"
 
 class AActor;
@@ -12,18 +12,20 @@ class UPropertyAnimatorCoreBase;
  * Preset for position properties (X, Y, Z) on scene component
  */
 UCLASS(Transient)
-class UPropertyAnimatorPresetScale : public UPropertyAnimatorCorePresetBase
+class UPropertyAnimatorPresetScale : public UPropertyAnimatorCorePropertyPreset
 {
 	GENERATED_BODY()
 
 public:
 	UPropertyAnimatorPresetScale()
-		: UPropertyAnimatorCorePresetBase(TEXT("Scale"))
-	{}
+	{
+		PresetName = TEXT("Scale");
+	}
 
 protected:
 	//~ Begin UPropertyAnimatorCorePresetBase
 	virtual void GetPresetProperties(const AActor* InActor, const UPropertyAnimatorCoreBase* InAnimator, TSet<FPropertyAnimatorCoreData>& OutProperties) const override;
 	virtual void OnPresetApplied(UPropertyAnimatorCoreBase* InAnimator, const TSet<FPropertyAnimatorCoreData>& InProperties) override;
+	virtual bool LoadPreset() override { return true; }
 	//~ End UPropertyAnimatorCorePresetBase
 };
\ No newline at end of file
diff --git a/Engine/Plugins/Experimental/PropertyAnimator/Source/PropertyAnimator/Private/Presets/PropertyAnimatorPresetText.h b/Engine/Plugins/Experimental/PropertyAnimator/Source/PropertyAnimator/Private/Presets/PropertyAnimatorPresetText.h
--- a/Engine/Plugins/Experimental/PropertyAnimator/Source/PropertyAnimator/Private/Presets/PropertyAnimatorPresetText.h
+++ b/Engine/Plugins/Experimental/PropertyAnimator/Source/PropertyAnimator/Private/Presets/PropertyAnimatorPresetText.h
@@ -2,7 +2,7 @@
 
 #pragma once
 
-#include "Presets/PropertyAnimatorCorePresetBase.h"
+#include "Presets/PropertyAnimatorCorePropertyPreset.h"
 #include "PropertyAnimatorPresetText.generated.h"
 
 class AActor;
@@ -12,18 +12,20 @@ class UPropertyAnimatorCoreBase;
  * Preset for text property on root text 3d component
  */
 UCLASS(Transient)
-class UPropertyAnimatorPresetText : public UPropertyAnimatorCorePresetBase
+class UPropertyAnimatorPresetText : public UPropertyAnimatorCorePropertyPreset
 {
 	GENERATED_BODY()
 
 public:
 	UPropertyAnimatorPresetText()
-		: UPropertyAnimatorCorePresetBase(TEXT("Text"))
-	{}
+	{
+		PresetName = TEXT("Text");
+	}
 
 protected:
 	//~ Begin UPropertyAnimatorCorePresetBase
 	virtual void GetPresetProperties(const AActor* InActor, const UPropertyAnimatorCoreBase* InAnimator, TSet<FPropertyAnimatorCoreData>& OutProperties) const override;
 	virtual void OnPresetApplied(UPropertyAnimatorCoreBase* InAnimator, const TSet<FPropertyAnimatorCoreData>& InProperties) override;
+	virtual bool LoadPreset() override { return true; }
 	//~ End UPropertyAnimatorCorePresetBase
 };
\ No newline at end of file
diff --git a/Engine/Plugins/Experimental/PropertyAnimator/Source/PropertyAnimator/Private/Presets/PropertyAnimatorPresetTextLocation.h b/Engine/Plugins/Experimental/PropertyAnimator/Source/PropertyAnimator/Private/Presets/PropertyAnimatorPresetTextLocation.h
--- a/Engine/Plugins/Experimental/PropertyAnimator/Source/PropertyAnimator/Private/Presets/PropertyAnimatorPresetTextLocation.h
+++ b/Engine/Plugins/Experimental/PropertyAnimator/Source/PropertyAnimator/Private/Presets/PropertyAnimatorPresetTextLocation.h
@@ -2,7 +2,7 @@
 
 #pragma once
 
-#include "Presets/PropertyAnimatorCorePresetBase.h"
+#include "Presets/PropertyAnimatorCorePropertyPreset.h"
 #include "PropertyAnimatorPresetTextLocation.generated.h"
 
 class AActor;
@@ -12,18 +12,20 @@ class UPropertyAnimatorCoreBase;
  * Preset for text character position properties (X, Y, Z) on scene component
  */
 UCLASS()
-class UPropertyAnimatorPresetTextLocation : public UPropertyAnimatorCorePresetBase
+class UPropertyAnimatorPresetTextLocation : public UPropertyAnimatorCorePropertyPreset
 {
 	GENERATED_BODY()
 
 public:
 	UPropertyAnimatorPresetTextLocation()
-		: UPropertyAnimatorCorePresetBase(TEXT("TextCharacterLocation"))
-	{}
+	{
+		PresetName = TEXT("TextCharacterLocation");
+	}
 
 protected:
 	//~ Begin UPropertyAnimatorCorePresetBase
 	virtual void GetPresetProperties(const AActor* InActor, const UPropertyAnimatorCoreBase* InAnimator, TSet<FPropertyAnimatorCoreData>& OutProperties) const override;
 	virtual void OnPresetApplied(UPropertyAnimatorCoreBase* InAnimator, const TSet<FPropertyAnimatorCoreData>& InProperties) override;
+	virtual bool LoadPreset() override { return true; }
 	//~ End UPropertyAnimatorCorePresetBase
 };
\ No newline at end of file
diff --git a/Engine/Plugins/Experimental/PropertyAnimator/Source/PropertyAnimator/Private/Presets/PropertyAnimatorPresetTextRotation.h b/Engine/Plugins/Experimental/PropertyAnimator/Source/PropertyAnimator/Private/Presets/PropertyAnimatorPresetTextRotation.h
--- a/Engine/Plugins/Experimental/PropertyAnimator/Source/PropertyAnimator/Private/Presets/PropertyAnimatorPresetTextRotation.h
+++ b/Engine/Plugins/Experimental/PropertyAnimator/Source/PropertyAnimator/Private/Presets/PropertyAnimatorPresetTextRotation.h
@@ -2,7 +2,7 @@
 
 #pragma once
 
-#include "Presets/PropertyAnimatorCorePresetBase.h"
+#include "Presets/PropertyAnimatorCorePropertyPreset.h"
 #include "PropertyAnimatorPresetTextRotation.generated.h"
 
 class AActor;
@@ -12,18 +12,20 @@ class UPropertyAnimatorCoreBase;
  * Preset for text character position properties (Roll, Pitch, Yaw) on scene component
  */
 UCLASS()
-class UPropertyAnimatorPresetTextRotation : public UPropertyAnimatorCorePresetBase
+class UPropertyAnimatorPresetTextRotation : public UPropertyAnimatorCorePropertyPreset
 {
 	GENERATED_BODY()
 
 public:
 	UPropertyAnimatorPresetTextRotation()
-		: UPropertyAnimatorCorePresetBase(TEXT("TextCharacterRotation"))
-	{}
+	{
+		PresetName = TEXT("TextCharacterRotation");
+	}
 
 protected:
 	//~ Begin UPropertyAnimatorCorePresetBase
 	virtual void GetPresetProperties(const AActor* InActor, const UPropertyAnimatorCoreBase* InAnimator, TSet<FPropertyAnimatorCoreData>& OutProperties) const override;
 	virtual void OnPresetApplied(UPropertyAnimatorCoreBase* InAnimator, const TSet<FPropertyAnimatorCoreData>& InProperties) override;
+	virtual bool LoadPreset() override { return true; }
 	//~ End UPropertyAnimatorCorePresetBase
 };
\ No newline at end of file
diff --git a/Engine/Plugins/Experimental/PropertyAnimator/Source/PropertyAnimator/Private/Presets/PropertyAnimatorPresetTextScale.h b/Engine/Plugins/Experimental/PropertyAnimator/Source/PropertyAnimator/Private/Presets/PropertyAnimatorPresetTextScale.h
--- a/Engine/Plugins/Experimental/PropertyAnimator/Source/PropertyAnimator/Private/Presets/PropertyAnimatorPresetTextScale.h
+++ b/Engine/Plugins/Experimental/PropertyAnimator/Source/PropertyAnimator/Private/Presets/PropertyAnimatorPresetTextScale.h
@@ -2,7 +2,7 @@
 
 #pragma once
 
-#include "Presets/PropertyAnimatorCorePresetBase.h"
+#include "Presets/PropertyAnimatorCorePropertyPreset.h"
 #include "PropertyAnimatorPresetTextScale.generated.h"
 
 class AActor;
@@ -12,18 +12,20 @@ class UPropertyAnimatorCoreBase;
  * Preset for text character position properties (X, Y, Z) on scene component
  */
 UCLASS()
-class UPropertyAnimatorPresetTextScale : public UPropertyAnimatorCorePresetBase
+class UPropertyAnimatorPresetTextScale : public UPropertyAnimatorCorePropertyPreset
 {
 	GENERATED_BODY()
 
 public:
 	UPropertyAnimatorPresetTextScale()
-		: UPropertyAnimatorCorePresetBase(TEXT("TextCharacterScale"))
-	{}
+	{
+		PresetName = TEXT("TextCharacterScale");
+	}
 
 protected:
 	//~ Begin UPropertyAnimatorCorePresetBase
 	virtual void GetPresetProperties(const AActor* InActor, const UPropertyAnimatorCoreBase* InAnimator, TSet<FPropertyAnimatorCoreData>& OutProperties) const override;
 	virtual void OnPresetApplied(UPropertyAnimatorCoreBase* InAnimator, const TSet<FPropertyAnimatorCoreData>& InProperties) override;
+	virtual bool LoadPreset() override { return true; }
 	//~ End UPropertyAnimatorCorePresetBase
 };
\ No newline at end of file
diff --git a/Engine/Plugins/Experimental/PropertyAnimator/Source/PropertyAnimator/Private/Presets/PropertyAnimatorPresetTextVisibility.h b/Engine/Plugins/Experimental/PropertyAnimator/Source/PropertyAnimator/Private/Presets/PropertyAnimatorPresetTextVisibility.h
--- a/Engine/Plugins/Experimental/PropertyAnimator/Source/PropertyAnimator/Private/Presets/PropertyAnimatorPresetTextVisibility.h
+++ b/Engine/Plugins/Experimental/PropertyAnimator/Source/PropertyAnimator/Private/Presets/PropertyAnimatorPresetTextVisibility.h
@@ -2,7 +2,7 @@
 
 #pragma once
 
-#include "Presets/PropertyAnimatorCorePresetBase.h"
+#include "Presets/PropertyAnimatorCorePropertyPreset.h"
 #include "PropertyAnimatorPresetTextVisibility.generated.h"
 
 class AActor;
@@ -12,18 +12,20 @@ class UPropertyAnimatorCoreBase;
  * Preset for text character visibility on scene component
  */
 UCLASS()
-class UPropertyAnimatorPresetTextVisibility : public UPropertyAnimatorCorePresetBase
+class UPropertyAnimatorPresetTextVisibility : public UPropertyAnimatorCorePropertyPreset
 {
 	GENERATED_BODY()
 
 public:
 	UPropertyAnimatorPresetTextVisibility()
-		: UPropertyAnimatorCorePresetBase(TEXT("TextCharacterVisibility"))
-	{}
+	{
+		PresetName = TEXT("TextCharacterVisibility");
+	}
 
 protected:
 	//~ Begin UPropertyAnimatorCorePresetBase
 	virtual void GetPresetProperties(const AActor* InActor, const UPropertyAnimatorCoreBase* InAnimator, TSet<FPropertyAnimatorCoreData>& OutProperties) const override;
 	virtual void OnPresetApplied(UPropertyAnimatorCoreBase* InAnimator, const TSet<FPropertyAnimatorCoreData>& InProperties) override;
+	virtual bool LoadPreset() override { return true; }
 	//~ End UPropertyAnimatorCorePresetBase
 };
\ No newline at end of file
diff --git a/Engine/Plugins/Experimental/PropertyAnimator/Source/PropertyAnimator/Private/Presets/PropertyAnimatorPresetVisibility.h b/Engine/Plugins/Experimental/PropertyAnimator/Source/PropertyAnimator/Private/Presets/PropertyAnimatorPresetVisibility.h
--- a/Engine/Plugins/Experimental/PropertyAnimator/Source/PropertyAnimator/Private/Presets/PropertyAnimatorPresetVisibility.h
+++ b/Engine/Plugins/Experimental/PropertyAnimator/Source/PropertyAnimator/Private/Presets/PropertyAnimatorPresetVisibility.h
@@ -2,7 +2,7 @@
 
 #pragma once
 
-#include "Presets/PropertyAnimatorCorePresetBase.h"
+#include "Presets/PropertyAnimatorCorePropertyPreset.h"
 #include "PropertyAnimatorPresetVisibility.generated.h"
 
 class AActor;
@@ -12,18 +12,20 @@ class UPropertyAnimatorCoreBase;
  * Preset for visibility properties on root scene component
  */
 UCLASS(Transient)
-class UPropertyAnimatorPresetVisibility : public UPropertyAnimatorCorePresetBase
+class UPropertyAnimatorPresetVisibility : public UPropertyAnimatorCorePropertyPreset
 {
 	GENERATED_BODY()
 
 public:
 	UPropertyAnimatorPresetVisibility()
-		: UPropertyAnimatorCorePresetBase(TEXT("Visibility"))
-	{}
+	{
+		PresetName = TEXT("Visibility");
+	}
 
 protected:
 	//~ Begin UPropertyAnimatorCorePresetBase
 	virtual void GetPresetProperties(const AActor* InActor, const UPropertyAnimatorCoreBase* InAnimator, TSet<FPropertyAnimatorCoreData>& OutProperties) const override;
 	virtual void OnPresetApplied(UPropertyAnimatorCoreBase* InAnimator, const TSet<FPropertyAnimatorCoreData>& InProperties) override;
+	virtual bool LoadPreset() override { return true; }
 	//~ End UPropertyAnimatorCorePresetBase
 };
\ No newline at end of file
diff --git a/Engine/Plugins/Experimental/PropertyAnimator/Source/PropertyAnimator/Private/Properties/PropertyAnimatorFloatContext.cpp b/Engine/Plugins/Experimental/PropertyAnimator/Source/PropertyAnimator/Private/Properties/PropertyAnimatorFloatContext.cpp
--- a/Engine/Plugins/Experimental/PropertyAnimator/Source/PropertyAnimator/Private/Properties/PropertyAnimatorFloatContext.cpp
+++ b/Engine/Plugins/Experimental/PropertyAnimator/Source/PropertyAnimator/Private/Properties/PropertyAnimatorFloatContext.cpp
@@ -3,6 +3,8 @@
 #include "Properties/PropertyAnimatorFloatContext.h"
 
 #include "Animators/PropertyAnimatorCoreBase.h"
+#include "Dom/JsonObject.h"
+#include "Dom/JsonValue.h"
 
 void UPropertyAnimatorFloatContext::SetAmplitudeMin(double InAmplitude)
 {
@@ -87,6 +89,41 @@ void UPropertyAnimatorFloatContext::OnAnimatedPropertyLinked()
 #endif
 }
 
+bool UPropertyAnimatorFloatContext::ImportPreset(const UPropertyAnimatorCorePresetBase* InPreset, const TSharedRef<FJsonValue>& InValue)
+{
+	TSharedPtr<FJsonObject>* JsonObject = nullptr;
+
+	if (Super::ImportPreset(InPreset, InValue) && InValue->TryGetObject(JsonObject))
+	{
+		double JsonAmplitudeMin = AmplitudeMin;
+		(*JsonObject)->TryGetNumberField(GET_MEMBER_NAME_STRING_CHECKED(UPropertyAnimatorFloatContext, AmplitudeMin), JsonAmplitudeMin);
+		SetAmplitudeMin(JsonAmplitudeMin);
+
+		double JsonAmplitudeMax = AmplitudeMax;
+		(*JsonObject)->TryGetNumberField(GET_MEMBER_NAME_STRING_CHECKED(UPropertyAnimatorFloatContext, AmplitudeMax), JsonAmplitudeMax);
+		SetAmplitudeMax(JsonAmplitudeMax);
+
+		return true;
+	}
+
+	return false;
+}
+
+bool UPropertyAnimatorFloatContext::ExportPreset(const UPropertyAnimatorCorePresetBase* InPreset, TSharedPtr<FJsonValue>& OutValue)
+{
+	TSharedPtr<FJsonObject>* JsonObject = nullptr;
+
+	if (Super::ExportPreset(InPreset, OutValue) && OutValue->TryGetObject(JsonObject))
+	{
+		(*JsonObject)->SetNumberField(GET_MEMBER_NAME_STRING_CHECKED(UPropertyAnimatorFloatContext, AmplitudeMin), AmplitudeMin);
+		(*JsonObject)->SetNumberField(GET_MEMBER_NAME_STRING_CHECKED(UPropertyAnimatorFloatContext, AmplitudeMax), AmplitudeMax);
+
+		return true;
+	}
+
+	return false;
+}
+
 double UPropertyAnimatorFloatContext::GetClampedAmplitude(double InAmplitude)
 {
 	if (AmplitudeClampMin.IsSet())
diff --git a/Engine/Plugins/Experimental/PropertyAnimator/Source/PropertyAnimator/Private/Properties/PropertyAnimatorRotatorContext.cpp b/Engine/Plugins/Experimental/PropertyAnimator/Source/PropertyAnimator/Private/Properties/PropertyAnimatorRotatorContext.cpp
--- a/Engine/Plugins/Experimental/PropertyAnimator/Source/PropertyAnimator/Private/Properties/PropertyAnimatorRotatorContext.cpp
+++ b/Engine/Plugins/Experimental/PropertyAnimator/Source/PropertyAnimator/Private/Properties/PropertyAnimatorRotatorContext.cpp
@@ -3,6 +3,8 @@
 #include "Properties/PropertyAnimatorRotatorContext.h"
 
 #include "Animators/PropertyAnimatorCoreBase.h"
+#include "Dom/JsonObject.h"
+#include "Dom/JsonValue.h"
 
 void UPropertyAnimatorRotatorContext::SetAmplitudeMin(const FRotator& InAmplitude)
 {
@@ -87,6 +89,45 @@ void UPropertyAnimatorRotatorContext::OnAnimatedPropertyLinked()
 #endif
 }
 
+bool UPropertyAnimatorRotatorContext::ImportPreset(const UPropertyAnimatorCorePresetBase* InPreset, const TSharedRef<FJsonValue>& InValue)
+{
+	TSharedPtr<FJsonObject>* JsonObject = nullptr;
+
+	if (Super::ImportPreset(InPreset, InValue) && InValue->TryGetObject(JsonObject))
+	{
+		FString JsonAmplitudeMin = AmplitudeMin.ToString();
+		(*JsonObject)->TryGetStringField(GET_MEMBER_NAME_STRING_CHECKED(UPropertyAnimatorRotatorContext, AmplitudeMin), JsonAmplitudeMin);
+		FRotator ParseAmplitudeMin;
+		ParseAmplitudeMin.InitFromString(JsonAmplitudeMin);
+		SetAmplitudeMin(ParseAmplitudeMin);
+
+		FString JsonAmplitudeMax = AmplitudeMax.ToString();
+		(*JsonObject)->TryGetStringField(GET_MEMBER_NAME_STRING_CHECKED(UPropertyAnimatorRotatorContext, AmplitudeMax), JsonAmplitudeMax);
+		FRotator ParseAmplitudeMax;
+		ParseAmplitudeMax.InitFromString(JsonAmplitudeMax);
+		SetAmplitudeMax(ParseAmplitudeMax);
+
+		return true;
+	}
+
+	return false;
+}
+
+bool UPropertyAnimatorRotatorContext::ExportPreset(const UPropertyAnimatorCorePresetBase* InPreset, TSharedPtr<FJsonValue>& OutValue)
+{
+	TSharedPtr<FJsonObject>* JsonObject = nullptr;
+
+	if (Super::ExportPreset(InPreset, OutValue) && OutValue->TryGetObject(JsonObject))
+	{
+		(*JsonObject)->SetStringField(GET_MEMBER_NAME_STRING_CHECKED(UPropertyAnimatorRotatorContext, AmplitudeMin), AmplitudeMin.ToString());
+		(*JsonObject)->SetStringField(GET_MEMBER_NAME_STRING_CHECKED(UPropertyAnimatorRotatorContext, AmplitudeMax), AmplitudeMax.ToString());
+
+		return true;
+	}
+
+	return false;
+}
+
 FRotator UPropertyAnimatorRotatorContext::GetClampedAmplitude(FRotator InAmplitude)
 {
 	if (AmplitudeClampMin.IsSet())
diff --git a/Engine/Plugins/Experimental/PropertyAnimator/Source/PropertyAnimator/Private/Properties/PropertyAnimatorVectorContext.cpp b/Engine/Plugins/Experimental/PropertyAnimator/Source/PropertyAnimator/Private/Properties/PropertyAnimatorVectorContext.cpp
--- a/Engine/Plugins/Experimental/PropertyAnimator/Source/PropertyAnimator/Private/Properties/PropertyAnimatorVectorContext.cpp
+++ b/Engine/Plugins/Experimental/PropertyAnimator/Source/PropertyAnimator/Private/Properties/PropertyAnimatorVectorContext.cpp
@@ -3,6 +3,8 @@
 #include "Properties/PropertyAnimatorVectorContext.h"
 
 #include "Animators/PropertyAnimatorCoreBase.h"
+#include "Dom/JsonObject.h"
+#include "Dom/JsonValue.h"
 
 void UPropertyAnimatorVectorContext::SetAmplitudeMin(const FVector& InAmplitude)
 {
@@ -87,6 +89,45 @@ void UPropertyAnimatorVectorContext::OnAnimatedPropertyLinked()
 #endif
 }
 
+bool UPropertyAnimatorVectorContext::ImportPreset(const UPropertyAnimatorCorePresetBase* InPreset, const TSharedRef<FJsonValue>& InValue)
+{
+	TSharedPtr<FJsonObject>* JsonObject = nullptr;
+
+	if (Super::ImportPreset(InPreset, InValue) && InValue->TryGetObject(JsonObject))
+	{
+		FString JsonAmplitudeMin = AmplitudeMin.ToString();
+		(*JsonObject)->TryGetStringField(GET_MEMBER_NAME_STRING_CHECKED(UPropertyAnimatorVectorContext, AmplitudeMin), JsonAmplitudeMin);
+		FVector ParseAmplitudeMin;
+		ParseAmplitudeMin.InitFromString(JsonAmplitudeMin);
+		SetAmplitudeMin(ParseAmplitudeMin);
+
+		FString JsonAmplitudeMax = AmplitudeMax.ToString();
+		(*JsonObject)->TryGetStringField(GET_MEMBER_NAME_STRING_CHECKED(UPropertyAnimatorVectorContext, AmplitudeMax), JsonAmplitudeMax);
+		FVector ParseAmplitudeMax;
+		ParseAmplitudeMax.InitFromString(JsonAmplitudeMax);
+		SetAmplitudeMax(ParseAmplitudeMax);
+
+		return true;
+	}
+
+	return false;
+}
+
+bool UPropertyAnimatorVectorContext::ExportPreset(const UPropertyAnimatorCorePresetBase* InPreset, TSharedPtr<FJsonValue>& OutValue)
+{
+	TSharedPtr<FJsonObject>* JsonObject = nullptr;
+
+	if (Super::ExportPreset(InPreset, OutValue) && OutValue->TryGetObject(JsonObject))
+	{
+		(*JsonObject)->SetStringField(GET_MEMBER_NAME_STRING_CHECKED(UPropertyAnimatorVectorContext, AmplitudeMin), AmplitudeMin.ToString());
+		(*JsonObject)->SetStringField(GET_MEMBER_NAME_STRING_CHECKED(UPropertyAnimatorVectorContext, AmplitudeMax), AmplitudeMax.ToString());
+
+		return true;
+	}
+
+	return false;
+}
+
 FVector UPropertyAnimatorVectorContext::GetClampedAmplitude(FVector InAmplitude)
 {
 	if (AmplitudeClampMin.IsSet())
diff --git a/Engine/Plugins/Experimental/PropertyAnimator/Source/PropertyAnimator/PropertyAnimator.Build.cs b/Engine/Plugins/Experimental/PropertyAnimator/Source/PropertyAnimator/PropertyAnimator.Build.cs
--- a/Engine/Plugins/Experimental/PropertyAnimator/Source/PropertyAnimator/PropertyAnimator.Build.cs
+++ b/Engine/Plugins/Experimental/PropertyAnimator/Source/PropertyAnimator/PropertyAnimator.Build.cs
@@ -24,6 +24,7 @@ public class PropertyAnimator : ModuleRules
 				"CoreUObject",
 				"DeveloperSettings",
 				"Engine",
+				"Json",
 				"MovieSceneTracks",
 				"Text3D"
 			});
diff --git a/Engine/Plugins/Experimental/PropertyAnimator/Source/PropertyAnimator/Public/Animators/PropertyAnimatorClock.h b/Engine/Plugins/Experimental/PropertyAnimator/Source/PropertyAnimator/Public/Animators/PropertyAnimatorClock.h
--- a/Engine/Plugins/Experimental/PropertyAnimator/Source/PropertyAnimator/Public/Animators/PropertyAnimatorClock.h
+++ b/Engine/Plugins/Experimental/PropertyAnimator/Source/PropertyAnimator/Public/Animators/PropertyAnimatorClock.h
@@ -42,6 +42,8 @@ protected:
 	//~ Begin UPropertyAnimatorCoreBase
 	virtual void OnAnimatorRegistered(FPropertyAnimatorCoreMetadata& InMetadata) override;
 	virtual void EvaluateProperties(FInstancedPropertyBag& InParameters) override;
+	virtual bool ImportPreset(const UPropertyAnimatorCorePresetBase* InPreset, const TSharedRef<FJsonValue>& InValue) override;
+	virtual bool ExportPreset(const UPropertyAnimatorCorePresetBase* InPreset, TSharedPtr<FJsonValue>& OutValue) override;
 	//~ End UPropertyAnimatorCoreBase
 
 	/**
diff --git a/Engine/Plugins/Experimental/PropertyAnimator/Source/PropertyAnimator/Public/Animators/PropertyAnimatorCounter.h b/Engine/Plugins/Experimental/PropertyAnimator/Source/PropertyAnimator/Public/Animators/PropertyAnimatorCounter.h
--- a/Engine/Plugins/Experimental/PropertyAnimator/Source/PropertyAnimator/Public/Animators/PropertyAnimatorCounter.h
+++ b/Engine/Plugins/Experimental/PropertyAnimator/Source/PropertyAnimator/Public/Animators/PropertyAnimatorCounter.h
@@ -126,6 +126,7 @@ public:
 
 protected:
 	FString FormatNumber(double InNumber) const;
+	const FPropertyAnimatorCounterFormat* GetFormat() const;
 
 	//~ Begin UObject
 #if WITH_EDITOR
@@ -136,6 +137,8 @@ protected:
 	//~ Begin UPropertyAnimatorCoreBase
 	virtual void OnAnimatorRegistered(FPropertyAnimatorCoreMetadata& InMetadata) override;
 	virtual void EvaluateProperties(FInstancedPropertyBag& InParameters) override;
+	virtual bool ImportPreset(const UPropertyAnimatorCorePresetBase* InPreset, const TSharedRef<FJsonValue>& InValue) override;
+	virtual bool ExportPreset(const UPropertyAnimatorCorePresetBase* InPreset, TSharedPtr<FJsonValue>& OutValue) override;
 	//~ End UPropertyAnimatorCoreBase
 
 	UFUNCTION()
@@ -152,7 +155,7 @@ protected:
 	void OnUseCustomFormatChanged();
 	void OnCustomFormatChanged();
 
-	/** Display pattern for the output, use {0} as replacement symbol */
+	/** Display pattern for the output to add prefix and suffix, use {0} as replacement symbol */
 	UPROPERTY(EditInstanceOnly, Setter, Getter, Category="Animator")
 	FText DisplayPattern = FText::FromString(TEXT("{0}"));
 
diff --git a/Engine/Plugins/Experimental/PropertyAnimator/Source/PropertyAnimator/Public/Animators/PropertyAnimatorCurve.h b/Engine/Plugins/Experimental/PropertyAnimator/Source/PropertyAnimator/Public/Animators/PropertyAnimatorCurve.h
--- a/Engine/Plugins/Experimental/PropertyAnimator/Source/PropertyAnimator/Public/Animators/PropertyAnimatorCurve.h
+++ b/Engine/Plugins/Experimental/PropertyAnimator/Source/PropertyAnimator/Public/Animators/PropertyAnimatorCurve.h
@@ -76,6 +76,8 @@ protected:
 	//~ Begin UPropertyAnimatorFloatBase
 	virtual void OnAnimatorRegistered(FPropertyAnimatorCoreMetadata& InMetadata) override;
 	virtual bool EvaluateProperty(const FPropertyAnimatorCoreData& InPropertyData, UPropertyAnimatorCoreContext* InContext, FInstancedPropertyBag& InParameters, FInstancedPropertyBag& OutEvaluationResult) const override;
+	virtual bool ImportPreset(const UPropertyAnimatorCorePresetBase* InPreset, const TSharedRef<FJsonValue>& InValue) override;
+	virtual bool ExportPreset(const UPropertyAnimatorCorePresetBase* InPreset, TSharedPtr<FJsonValue>& OutValue) override;
 	//~ End UPropertyAnimatorFloatBase
 
 	void OnEaseInChanged();
diff --git a/Engine/Plugins/Experimental/PropertyAnimator/Source/PropertyAnimator/Public/Animators/PropertyAnimatorNumericBase.h b/Engine/Plugins/Experimental/PropertyAnimator/Source/PropertyAnimator/Public/Animators/PropertyAnimatorNumericBase.h
--- a/Engine/Plugins/Experimental/PropertyAnimator/Source/PropertyAnimator/Public/Animators/PropertyAnimatorNumericBase.h
+++ b/Engine/Plugins/Experimental/PropertyAnimator/Source/PropertyAnimator/Public/Animators/PropertyAnimatorNumericBase.h
@@ -59,12 +59,6 @@ public:
 		return CycleRate;
 	}
 
-	PROPERTYANIMATOR_API void SetTimeOffset(double InOffset);
-	double GetTimeOffset() const
-	{
-		return TimeOffset;
-	}
-
 	PROPERTYANIMATOR_API void SetRandomTimeOffset(bool bInOffset);
 	bool GetRandomTimeOffset() const
 	{
@@ -87,7 +81,6 @@ protected:
 	virtual void OnMagnitudeChanged() {}
 	virtual void OnCycleDurationChanged() {}
 	virtual void OnCycleModeChanged() {}
-	virtual void OnTimeOffsetChanged() {}
 	virtual void OnSeedChanged() {}
 
 	//~ Begin UPropertyAnimatorCoreBase
@@ -97,6 +90,8 @@ protected:
 	virtual void OnPropertyLinked(UPropertyAnimatorCoreContext* InLinkedProperty, EPropertyAnimatorPropertySupport InSupport) override;
 	virtual bool IsTimeSourceSupported(UPropertyAnimatorCoreTimeSourceBase* InTimeSource) const override;
 	virtual void OnAnimatorRegistered(FPropertyAnimatorCoreMetadata& InMetadata) override;
+	virtual bool ImportPreset(const UPropertyAnimatorCorePresetBase* InPreset, const TSharedRef<FJsonValue>& InValue) override;
+	virtual bool ExportPreset(const UPropertyAnimatorCorePresetBase* InPreset, TSharedPtr<FJsonValue>& OutValue) override;
 	//~ End UPropertyAnimatorCoreBase
 
 	/** Evaluate and return float value for a property */
@@ -133,10 +128,6 @@ protected:
 	UPROPERTY(EditInstanceOnly, Setter, Getter, Category="Animator", meta=(EditCondition="bRandomTimeOffset"))
 	int32 Seed = 0;
 
-	/** Time offset accumulated for each property for every round */
-	UPROPERTY(EditInstanceOnly, Setter, Getter, Category="Animator", meta=(Units=Seconds))
-	double TimeOffset = 0;
-
 private:
 	/** Random stream for time offset */
 	FRandomStream RandomStream = FRandomStream(Seed);
diff --git a/Engine/Plugins/Experimental/PropertyAnimator/Source/PropertyAnimator/Public/Animators/PropertyAnimatorSoundWave.h b/Engine/Plugins/Experimental/PropertyAnimator/Source/PropertyAnimator/Public/Animators/PropertyAnimatorSoundWave.h
--- a/Engine/Plugins/Experimental/PropertyAnimator/Source/PropertyAnimator/Public/Animators/PropertyAnimatorSoundWave.h
+++ b/Engine/Plugins/Experimental/PropertyAnimator/Source/PropertyAnimator/Public/Animators/PropertyAnimatorSoundWave.h
@@ -43,6 +43,8 @@ protected:
 	//~ Begin UPropertyAnimatorFloatBase
 	virtual void OnAnimatorRegistered(FPropertyAnimatorCoreMetadata& InMetadata) override;
 	virtual bool EvaluateProperty(const FPropertyAnimatorCoreData& InPropertyData, UPropertyAnimatorCoreContext* InContext, FInstancedPropertyBag& InParameters, FInstancedPropertyBag& OutEvaluationResult) const override;
+	virtual bool ImportPreset(const UPropertyAnimatorCorePresetBase* InPreset, const TSharedRef<FJsonValue>& InValue) override;
+	virtual bool ExportPreset(const UPropertyAnimatorCorePresetBase* InPreset, TSharedPtr<FJsonValue>& OutValue) override;
 	//~ End UPropertyAnimatorFloatBase
 
 	/**
diff --git a/Engine/Plugins/Experimental/PropertyAnimator/Source/PropertyAnimator/Public/Properties/PropertyAnimatorFloatContext.h b/Engine/Plugins/Experimental/PropertyAnimator/Source/PropertyAnimator/Public/Properties/PropertyAnimatorFloatContext.h
--- a/Engine/Plugins/Experimental/PropertyAnimator/Source/PropertyAnimator/Public/Properties/PropertyAnimatorFloatContext.h
+++ b/Engine/Plugins/Experimental/PropertyAnimator/Source/PropertyAnimator/Public/Properties/PropertyAnimatorFloatContext.h
@@ -34,6 +34,8 @@ protected:
 	//~ Begin UPropertyAnimatorCoreContext
 	virtual bool EvaluateProperty(const FPropertyAnimatorCoreData& InProperty, const FInstancedPropertyBag& InAnimatorResult, FInstancedPropertyBag& OutEvaluatedValues) override;
 	virtual void OnAnimatedPropertyLinked() override;
+	virtual bool ImportPreset(const UPropertyAnimatorCorePresetBase* InPreset, const TSharedRef<FJsonValue>& InValue) override;
+	virtual bool ExportPreset(const UPropertyAnimatorCorePresetBase* InPreset, TSharedPtr<FJsonValue>& OutValue) override;
 	//~ End UPropertyAnimatorCoreContext
 
 	double GetClampedAmplitude(double InAmplitude);
diff --git a/Engine/Plugins/Experimental/PropertyAnimator/Source/PropertyAnimator/Public/Properties/PropertyAnimatorRotatorContext.h b/Engine/Plugins/Experimental/PropertyAnimator/Source/PropertyAnimator/Public/Properties/PropertyAnimatorRotatorContext.h
--- a/Engine/Plugins/Experimental/PropertyAnimator/Source/PropertyAnimator/Public/Properties/PropertyAnimatorRotatorContext.h
+++ b/Engine/Plugins/Experimental/PropertyAnimator/Source/PropertyAnimator/Public/Properties/PropertyAnimatorRotatorContext.h
@@ -34,6 +34,8 @@ protected:
 	//~ Begin UPropertyAnimatorCoreContext
 	virtual bool EvaluateProperty(const FPropertyAnimatorCoreData& InProperty, const FInstancedPropertyBag& InAnimatorResult, FInstancedPropertyBag& OutEvaluatedValues) override;
 	virtual void OnAnimatedPropertyLinked() override;
+	virtual bool ImportPreset(const UPropertyAnimatorCorePresetBase* InPreset, const TSharedRef<FJsonValue>& InValue) override;
+	virtual bool ExportPreset(const UPropertyAnimatorCorePresetBase* InPreset, TSharedPtr<FJsonValue>& OutValue) override;
 	//~ End UPropertyAnimatorCoreContext
 
 	FRotator GetClampedAmplitude(FRotator InAmplitude);
diff --git a/Engine/Plugins/Experimental/PropertyAnimator/Source/PropertyAnimator/Public/Properties/PropertyAnimatorVectorContext.h b/Engine/Plugins/Experimental/PropertyAnimator/Source/PropertyAnimator/Public/Properties/PropertyAnimatorVectorContext.h
--- a/Engine/Plugins/Experimental/PropertyAnimator/Source/PropertyAnimator/Public/Properties/PropertyAnimatorVectorContext.h
+++ b/Engine/Plugins/Experimental/PropertyAnimator/Source/PropertyAnimator/Public/Properties/PropertyAnimatorVectorContext.h
@@ -34,6 +34,8 @@ protected:
 	//~ Begin UPropertyAnimatorCoreContext
 	virtual bool EvaluateProperty(const FPropertyAnimatorCoreData& InProperty, const FInstancedPropertyBag& InAnimatorResult, FInstancedPropertyBag& OutEvaluatedValues) override;
 	virtual void OnAnimatedPropertyLinked() override;
+	virtual bool ImportPreset(const UPropertyAnimatorCorePresetBase* InPreset, const TSharedRef<FJsonValue>& InValue) override;
+	virtual bool ExportPreset(const UPropertyAnimatorCorePresetBase* InPreset, TSharedPtr<FJsonValue>& OutValue) override;
 	//~ End UPropertyAnimatorCoreContext
 
 	FVector GetClampedAmplitude(FVector InAmplitude);
diff --git a/Engine/Plugins/Experimental/PropertyAnimatorCore/Source/PropertyAnimatorCore/Private/Animators/PropertyAnimatorCoreBase.cpp b/Engine/Plugins/Experimental/PropertyAnimatorCore/Source/PropertyAnimatorCore/Private/Animators/PropertyAnimatorCoreBase.cpp
--- a/Engine/Plugins/Experimental/PropertyAnimatorCore/Source/PropertyAnimatorCore/Private/Animators/PropertyAnimatorCoreBase.cpp
+++ b/Engine/Plugins/Experimental/PropertyAnimatorCore/Source/PropertyAnimatorCore/Private/Animators/PropertyAnimatorCoreBase.cpp
@@ -3,6 +3,7 @@
 #include "Animators/PropertyAnimatorCoreBase.h"
 
 #include "Components/PropertyAnimatorCoreComponent.h"
+#include "Dom/JsonObject.h"
 #include "GameFramework/Actor.h"
 #include "Properties/PropertyAnimatorCoreGroupBase.h"
 #include "Properties/PropertyAnimatorCoreResolver.h"
@@ -211,6 +212,108 @@ void UPropertyAnimatorCoreBase::PostEditChangeProperty(FPropertyChangedEvent& Pr
 }
 #endif
 
+bool UPropertyAnimatorCoreBase::ImportPreset(const UPropertyAnimatorCorePresetBase* InPreset, const TSharedRef<FJsonValue>& InValue)
+{
+	const TSharedPtr<FJsonObject>* JsonAnimatorObject;
+	if (!InValue->TryGetObject(JsonAnimatorObject) || !JsonAnimatorObject)
+	{
+		return false;
+	}
+
+	bool bJsonEnabled = bAnimatorEnabled;
+	(*JsonAnimatorObject)->TryGetBoolField(GET_MEMBER_NAME_STRING_CHECKED(UPropertyAnimatorCoreBase, bAnimatorEnabled), bJsonEnabled);
+	SetAnimatorEnabled(bJsonEnabled);
+
+	FString JsonDisplayName = AnimatorDisplayName.ToString();
+	(*JsonAnimatorObject)->TryGetStringField(GET_MEMBER_NAME_STRING_CHECKED(UPropertyAnimatorCoreBase, AnimatorDisplayName), JsonDisplayName);
+	SetAnimatorDisplayName(FName(JsonDisplayName));
+
+	const TArray<TSharedPtr<FJsonValue>>* JsonValues;
+	(*JsonAnimatorObject)->TryGetArrayField(GET_MEMBER_NAME_STRING_CHECKED(UPropertyAnimatorCoreBase, LinkedProperties), JsonValues);
+	for (const TSharedPtr<FJsonValue>& JsonValue : (*JsonValues))
+	{
+		TSharedPtr<FJsonObject>* JsonPropertyObject;
+		if (!JsonValue->TryGetObject(JsonPropertyObject) || !JsonPropertyObject || !JsonPropertyObject->IsValid())
+		{
+			continue;
+		}
+
+		FString AnimatedPropertyLocatorPath;
+		if (!(*JsonPropertyObject)->TryGetStringField(GET_MEMBER_NAME_STRING_CHECKED(UPropertyAnimatorCoreContext, AnimatedProperty), AnimatedPropertyLocatorPath))
+		{
+			continue;
+		}
+
+		FPropertyAnimatorCoreData PropertyData(GetAnimatorActor(), AnimatedPropertyLocatorPath);
+
+		if (!PropertyData.IsResolved())
+		{
+			continue;
+		}
+
+		if (IPropertyAnimatorCorePresetable* PropertyContext = Cast<IPropertyAnimatorCorePresetable>(LinkProperty(PropertyData)))
+		{
+			PropertyContext->ImportPreset(InPreset, JsonValue.ToSharedRef());
+		}
+	}
+
+	bool bJsonOverrideTimeSource = bOverrideTimeSource;
+	(*JsonAnimatorObject)->TryGetBoolField(GET_MEMBER_NAME_STRING_CHECKED(UPropertyAnimatorCoreBase, bOverrideTimeSource), bJsonOverrideTimeSource);
+	SetOverrideTimeSource(bJsonOverrideTimeSource);
+
+	FString JsonTimeSourceName = TimeSourceName.ToString();
+	(*JsonAnimatorObject)->TryGetStringField(GET_MEMBER_NAME_STRING_CHECKED(UPropertyAnimatorCoreBase, TimeSourceName), JsonTimeSourceName);
+	SetTimeSourceName(FName(JsonTimeSourceName));
+
+	if (UPropertyAnimatorCoreTimeSourceBase* TimeSource = FindOrAddTimeSource(GetTimeSourceName()))
+	{
+		if ((*JsonAnimatorObject)->HasTypedField(GET_MEMBER_NAME_STRING_CHECKED(UPropertyAnimatorCoreBase, ActiveTimeSource), EJson::Object))
+		{
+			TimeSource->ImportPreset(InPreset, (*JsonAnimatorObject)->TryGetField(GET_MEMBER_NAME_STRING_CHECKED(UPropertyAnimatorCoreBase, ActiveTimeSource)).ToSharedRef());
+		}
+	}
+
+	return true;
+}
+
+bool UPropertyAnimatorCoreBase::ExportPreset(const UPropertyAnimatorCorePresetBase* InPreset, TSharedPtr<FJsonValue>& OutValue)
+{
+	TSharedRef<FJsonObject> JsonAnimatorObject = MakeShared<FJsonObject>();
+	OutValue = MakeShared<FJsonValueObject>(JsonAnimatorObject);
+
+	JsonAnimatorObject->SetStringField(TEXT("AnimatorClass"), GetClass()->GetClassPathName().ToString());
+	JsonAnimatorObject->SetBoolField(GET_MEMBER_NAME_STRING_CHECKED(UPropertyAnimatorCoreBase, bAnimatorEnabled), bAnimatorEnabled);
+	JsonAnimatorObject->SetStringField(GET_MEMBER_NAME_STRING_CHECKED(UPropertyAnimatorCoreBase, AnimatorDisplayName), AnimatorDisplayName.ToString());
+
+	TArray<TSharedPtr<FJsonValue>> JsonValues;
+	for (const TObjectPtr<UPropertyAnimatorCoreContext>& LinkedProperty : LinkedProperties)
+	{
+		if (IPropertyAnimatorCorePresetable* PropertyContext = Cast<IPropertyAnimatorCorePresetable>(LinkedProperty))
+		{
+			TSharedPtr<FJsonValue> JsonValue;
+			if (PropertyContext->ExportPreset(InPreset, JsonValue) && JsonValue.IsValid())
+			{
+				JsonValues.Add(JsonValue);
+			}
+		}
+	}
+	JsonAnimatorObject->SetArrayField(GET_MEMBER_NAME_STRING_CHECKED(UPropertyAnimatorCoreBase, LinkedProperties), JsonValues);
+
+	JsonAnimatorObject->SetBoolField(GET_MEMBER_NAME_STRING_CHECKED(UPropertyAnimatorCoreBase, bOverrideTimeSource), bOverrideTimeSource);
+	JsonAnimatorObject->SetStringField(GET_MEMBER_NAME_STRING_CHECKED(UPropertyAnimatorCoreBase, TimeSourceName), TimeSourceName.ToString());
+
+	if (UPropertyAnimatorCoreTimeSourceBase* TimeSource = GetActiveTimeSource())
+	{
+		TSharedPtr<FJsonValue> JsonValue;
+		if (TimeSource->ExportPreset(InPreset, JsonValue) && JsonValue.IsValid())
+		{
+			JsonAnimatorObject->SetField(GET_MEMBER_NAME_STRING_CHECKED(UPropertyAnimatorCoreBase, ActiveTimeSource), JsonValue);
+		}
+	}
+
+	return true;
+}
+
 AActor* UPropertyAnimatorCoreBase::GetAnimatorActor() const
 {
 	return GetTypedOuter<AActor>();
@@ -235,7 +338,7 @@ void UPropertyAnimatorCoreBase::SetOverrideTimeSource(bool bInOverride)
 	}
 
 	bOverrideTimeSource = bInOverride;
-	OnTimeSourceChanged();
+	OnTimeSourceNameChanged();
 }
 
 void UPropertyAnimatorCoreBase::SetTimeSourceName(FName InTimeSourceName)
@@ -431,7 +534,7 @@ void UPropertyAnimatorCoreBase::OnTimeSourceNameChanged()
 		ActiveTimeSource->DeactivateTimeSource();
 	}
 
-	ActiveTimeSource = FindOrAddTimeSource(TimeSourceName);
+	ActiveTimeSource = bOverrideTimeSource ? FindOrAddTimeSource(TimeSourceName) : nullptr;
 
 	if (ActiveTimeSource)
 	{
@@ -441,6 +544,11 @@ void UPropertyAnimatorCoreBase::OnTimeSourceNameChanged()
 	OnTimeSourceChanged();
 }
 
+void UPropertyAnimatorCoreBase::OnTimeSourceEnterIdleState()
+{
+	RestoreProperties(/** Force */true);
+}
+
 void UPropertyAnimatorCoreBase::ResolvePropertiesOwner(AActor* InNewOwner)
 {
 	// Resolve linked properties against current actor
@@ -586,7 +694,7 @@ TArray<FName> UPropertyAnimatorCoreBase::GetTimeSourceNames() const
 
 UPropertyAnimatorCoreTimeSourceBase* UPropertyAnimatorCoreBase::FindOrAddTimeSource(FName InTimeSourceName)
 {
-	if (IsTemplate() || !bOverrideTimeSource)
+	if (IsTemplate())
 	{
 		return nullptr;
 	}
@@ -672,11 +780,13 @@ int32 UPropertyAnimatorCoreBase::GetLinkedPropertiesCount() const
 	return LinkedProperties.Num();
 }
 
-bool UPropertyAnimatorCoreBase::LinkProperty(const FPropertyAnimatorCoreData& InLinkProperty)
+UPropertyAnimatorCoreContext* UPropertyAnimatorCoreBase::LinkProperty(const FPropertyAnimatorCoreData& InLinkProperty)
 {
+	UPropertyAnimatorCoreContext* PropertyContext = nullptr;
+
 	if (!InLinkProperty.IsResolved())
 	{
-		return false;
+		return PropertyContext;
 	}
 
 	const UObject* Owner = InLinkProperty.GetOwner();
@@ -684,19 +794,19 @@ bool UPropertyAnimatorCoreBase::LinkProperty(const FPropertyAnimatorCoreData& In
 
 	if (Owner != OwningActor && !Owner->IsIn(OwningActor))
 	{
-		return false;
+		return PropertyContext;
 	}
 
 	const EPropertyAnimatorPropertySupport Support = GetPropertySupport(InLinkProperty);
 
 	if (Support == EPropertyAnimatorPropertySupport::None)
 	{
-		return false;
+		return PropertyContext;
 	}
 
 	if (IsPropertyLinked(InLinkProperty))
 	{
-		return false;
+		return GetLinkedPropertyContext(InLinkProperty);
 	}
 
 	const TSubclassOf<UPropertyAnimatorCoreContext> ContextSubclass = GetPropertyContextClass(InLinkProperty);
@@ -704,10 +814,10 @@ bool UPropertyAnimatorCoreBase::LinkProperty(const FPropertyAnimatorCoreData& In
 
 	if (!IsValid(ContextClass))
 	{
-		return false;
+		return PropertyContext;
 	}
 
-	UPropertyAnimatorCoreContext* PropertyContext = NewObject<UPropertyAnimatorCoreContext>(this, ContextClass, NAME_None, RF_Transactional);
+	PropertyContext = NewObject<UPropertyAnimatorCoreContext>(this, ContextClass, NAME_None, RF_Transactional);
 	PropertyContext->ConstructInternal(InLinkProperty);
 
 	LinkedProperties.Add(PropertyContext);
@@ -715,7 +825,7 @@ bool UPropertyAnimatorCoreBase::LinkProperty(const FPropertyAnimatorCoreData& In
 
 	UPropertyAnimatorCoreBase::OnAnimatorPropertyLinkedDelegate.Broadcast(this, InLinkProperty);
 
-	return true;
+	return PropertyContext;
 }
 
 bool UPropertyAnimatorCoreBase::UnlinkProperty(const FPropertyAnimatorCoreData& InUnlinkProperty)
diff --git a/Engine/Plugins/Experimental/PropertyAnimatorCore/Source/PropertyAnimatorCore/Private/Components/PropertyAnimatorCoreComponent.cpp b/Engine/Plugins/Experimental/PropertyAnimatorCore/Source/PropertyAnimatorCore/Private/Components/PropertyAnimatorCoreComponent.cpp
--- a/Engine/Plugins/Experimental/PropertyAnimatorCore/Source/PropertyAnimatorCore/Private/Components/PropertyAnimatorCoreComponent.cpp
+++ b/Engine/Plugins/Experimental/PropertyAnimatorCore/Source/PropertyAnimatorCore/Private/Components/PropertyAnimatorCoreComponent.cpp
@@ -516,9 +516,9 @@ bool UPropertyAnimatorCoreComponent::EvaluateAnimators()
 
 		if (AnimatorTimeResult != EPropertyAnimatorCoreTimeSourceResult::Evaluate)
 		{
-			if (AnimatorTimeResult == EPropertyAnimatorCoreTimeSourceResult::Reset)
+			if (AnimatorTimeResult == EPropertyAnimatorCoreTimeSourceResult::Idle)
 			{
-				Animator->RestoreProperties(/** Force */true);
+				Animator->OnTimeSourceEnterIdleState();
 			}
 
 			continue;
diff --git a/Engine/Plugins/Experimental/PropertyAnimatorCore/Source/PropertyAnimatorCore/Private/Presets/PropertyAnimatorCoreAnimatorPreset.cpp b/Engine/Plugins/Experimental/PropertyAnimatorCore/Source/PropertyAnimatorCore/Private/Presets/PropertyAnimatorCoreAnimatorPreset.cpp
new file mode 100644
--- /dev/null
+++ b/Engine/Plugins/Experimental/PropertyAnimatorCore/Source/PropertyAnimatorCore/Private/Presets/PropertyAnimatorCoreAnimatorPreset.cpp
@@ -0,0 +1,93 @@
+// Copyright Epic Games, Inc. All Rights Reserved.
+
+#include "Presets/PropertyAnimatorCoreAnimatorPreset.h"
+
+#include "Animators/PropertyAnimatorCoreBase.h"
+#include "Dom/JsonObject.h"
+#include "Dom/JsonValue.h"
+#include "Presets/PropertyAnimatorCorePresetable.h"
+#include "Serialization/JsonReader.h"
+#include "Serialization/JsonSerializer.h"
+#include "Serialization/JsonWriter.h"
+
+bool UPropertyAnimatorCoreAnimatorPreset::IsPresetApplied(const UPropertyAnimatorCoreBase* InAnimator) const
+{
+	return false;
+}
+
+bool UPropertyAnimatorCoreAnimatorPreset::IsPresetSupported(const AActor* InActor, const UPropertyAnimatorCoreBase* InAnimator) const
+{
+	return InAnimator && InAnimator->IsA(TargetAnimatorClass);
+}
+
+bool UPropertyAnimatorCoreAnimatorPreset::ApplyPreset(UPropertyAnimatorCoreBase* InAnimator)
+{
+	if (InAnimator->IsTemplate())
+	{
+		return false;
+	}
+
+	return InAnimator->ImportPreset(this, AnimatorPreset.ToSharedRef());
+}
+
+bool UPropertyAnimatorCoreAnimatorPreset::UnapplyPreset(UPropertyAnimatorCoreBase* InAnimator)
+{
+	return false;
+}
+
+void UPropertyAnimatorCoreAnimatorPreset::CreatePreset(FName InName, const TArray<IPropertyAnimatorCorePresetable*>& InPresetableItem)
+{
+	Super::CreatePreset(InName, InPresetableItem);
+
+	TSharedPtr<FJsonValue> JsonValue;
+
+	if (InPresetableItem[0]
+		&& InPresetableItem[0]->ExportPreset(this, JsonValue)
+		&& JsonValue.IsValid())
+	{
+		TSharedPtr<FJsonObject>* JsonObject = nullptr;
+		if (JsonValue->TryGetObject(JsonObject))
+		{
+			FString JsonString;
+			TSharedRef<TJsonWriter<>> JsonWriter = TJsonWriterFactory<>::Create(&JsonString);
+			if (FJsonSerializer::Serialize(JsonObject->ToSharedRef(), JsonWriter))
+			{
+				PresetVersion = 0;
+				PresetContent = JsonString;
+			}
+		}
+	}
+}
+
+bool UPropertyAnimatorCoreAnimatorPreset::LoadPreset()
+{
+	if (PresetContent.IsEmpty())
+	{
+		return false;
+	}
+
+	TSharedRef<TJsonReader<>> Reader = TJsonReaderFactory<>::Create(PresetContent);
+
+	TSharedPtr<FJsonObject> JsonObject;
+	if (!FJsonSerializer::Deserialize(Reader, JsonObject) || !JsonObject.IsValid())
+	{
+		return false;
+	}
+
+	FString Class;
+	JsonObject->TryGetStringField(TEXT("AnimatorClass"), Class);
+
+	if (UClass* AnimatorClass = LoadObject<UClass>(nullptr, *Class))
+	{
+		TargetAnimatorClass = AnimatorClass;
+		AnimatorPreset = MakeShared<FJsonValueObject>(JsonObject);
+		return true;
+	}
+
+	return false;
+}
+
+UPropertyAnimatorCoreBase* UPropertyAnimatorCoreAnimatorPreset::GetAnimatorTemplate() const
+{
+	return TargetAnimatorClass.GetDefaultObject();
+}
diff --git a/Engine/Plugins/Experimental/PropertyAnimatorCore/Source/PropertyAnimatorCore/Private/Presets/PropertyAnimatorCorePresetBase.cpp b/Engine/Plugins/Experimental/PropertyAnimatorCore/Source/PropertyAnimatorCore/Private/Presets/PropertyAnimatorCorePresetBase.cpp
--- a/Engine/Plugins/Experimental/PropertyAnimatorCore/Source/PropertyAnimatorCore/Private/Presets/PropertyAnimatorCorePresetBase.cpp
+++ b/Engine/Plugins/Experimental/PropertyAnimatorCore/Source/PropertyAnimatorCore/Private/Presets/PropertyAnimatorCorePresetBase.cpp
@@ -3,128 +3,18 @@
 #include "Presets/PropertyAnimatorCorePresetBase.h"
 
 #include "Animators/PropertyAnimatorCoreBase.h"
+#include "Dom/JsonObject.h"
+#include "Dom/JsonValue.h"
+#include "Presets/PropertyAnimatorCorePresetable.h"
+#include "Serialization/JsonWriter.h"
+#include "Serialization/JsonSerializer.h"
 
 FString UPropertyAnimatorCorePresetBase::GetPresetDisplayName() const
 {
 	return FName::NameToDisplayString(PresetName.ToString(), false);
 }
 
-void UPropertyAnimatorCorePresetBase::GetSupportedPresetProperties(const AActor* InActor, const UPropertyAnimatorCoreBase* InAnimator, TSet<FPropertyAnimatorCoreData>& OutProperties) const
+void UPropertyAnimatorCorePresetBase::CreatePreset(FName InName, const TArray<IPropertyAnimatorCorePresetable*>& InPresetableItem)
 {
-	TSet<FPropertyAnimatorCoreData> PresetProperties;
-	GetPresetProperties(InActor, InAnimator, PresetProperties);
-	OutProperties.Empty(PresetProperties.Num());
-
-	if (PresetProperties.IsEmpty())
-	{
-		return;
-	}
-
-	for (const FPropertyAnimatorCoreData& PresetProperty : PresetProperties)
-	{
-		InAnimator->GetPropertiesSupported(PresetProperty, OutProperties, /** SearchDepth */3);
-	}
-}
-
-bool UPropertyAnimatorCorePresetBase::IsPresetSupported(const AActor* InActor, const UPropertyAnimatorCoreBase* InAnimator) const
-{
-	if (!IsValid(InActor) || !IsValid(InAnimator))
-	{
-		return false;
-	}
-
-	TSet<FPropertyAnimatorCoreData> SupportedProperties;
-	GetSupportedPresetProperties(InActor, InAnimator, SupportedProperties);
-
-	return !SupportedProperties.IsEmpty();
-}
-
-bool UPropertyAnimatorCorePresetBase::ApplyPreset(UPropertyAnimatorCoreBase* InAnimator)
-{
-	if (!IsValid(InAnimator) || InAnimator->IsTemplate())
-	{
-		return false;
-	}
-
-	TSet<FPropertyAnimatorCoreData> SupportedProperties;
-	GetSupportedPresetProperties(InAnimator->GetAnimatorActor(), InAnimator, SupportedProperties);
-
-	if (SupportedProperties.IsEmpty())
-	{
-		return false;
-	}
-
-	for (FPropertyAnimatorCoreData& SupportedProperty : SupportedProperties)
-	{
-		InAnimator->LinkProperty(SupportedProperty);
-	}
-
-	InAnimator->SetAnimatorDisplayName(FName(InAnimator->GetAnimatorOriginalName().ToString() + TEXT("_") + GetPresetDisplayName()));
-
-	OnPresetApplied(InAnimator, SupportedProperties);
-
-	return true;
-}
-
-bool UPropertyAnimatorCorePresetBase::IsPresetApplied(const UPropertyAnimatorCoreBase* InAnimator) const
-{
-	if (!IsValid(InAnimator) || InAnimator->IsTemplate())
-	{
-		return false;
-	}
-
-	TSet<FPropertyAnimatorCoreData> SupportedProperties;
-	GetSupportedPresetProperties(InAnimator->GetAnimatorActor(), InAnimator, SupportedProperties);
-
-	if (SupportedProperties.IsEmpty())
-	{
-		return false;
-	}
-
-	return InAnimator->IsPropertiesLinked(SupportedProperties);
-}
-
-bool UPropertyAnimatorCorePresetBase::UnapplyPreset(UPropertyAnimatorCoreBase* InAnimator)
-{
-	if (!IsValid(InAnimator) || InAnimator->IsTemplate())
-	{
-		return false;
-	}
-
-	TSet<FPropertyAnimatorCoreData> SupportedProperties;
-	GetSupportedPresetProperties(InAnimator->GetAnimatorActor(), InAnimator, SupportedProperties);
-
-	if (SupportedProperties.IsEmpty())
-	{
-		return false;
-	}
-
-	for (FPropertyAnimatorCoreData& SupportedProperty : SupportedProperties)
-	{
-		InAnimator->UnlinkProperty(SupportedProperty);
-	}
-
-	return true;
-}
-
-void UPropertyAnimatorCorePresetBase::GetAppliedPresetProperties(const UPropertyAnimatorCoreBase* InAnimator, TSet<FPropertyAnimatorCoreData>& OutSupportedProperties, TSet<FPropertyAnimatorCoreData>& OutAppliedProperties)
-{
-	OutSupportedProperties.Empty();
-	OutAppliedProperties.Empty();
-
-	if (!IsValid(InAnimator) || InAnimator->IsTemplate())
-	{
-		return;
-	}
-
-	GetSupportedPresetProperties(InAnimator->GetAnimatorActor(), InAnimator, OutSupportedProperties);
-	OutAppliedProperties.Reserve(OutSupportedProperties.Num());
-
-	for (const FPropertyAnimatorCoreData& Property : OutSupportedProperties)
-	{
-		if (InAnimator->IsPropertyLinked(Property))
-		{
-			OutAppliedProperties.Add(Property);
-		}
-	}
+	PresetName = InName;
 }
diff --git a/Engine/Plugins/Experimental/PropertyAnimatorCore/Source/PropertyAnimatorCore/Private/Presets/PropertyAnimatorCorePropertyPreset.cpp b/Engine/Plugins/Experimental/PropertyAnimatorCore/Source/PropertyAnimatorCore/Private/Presets/PropertyAnimatorCorePropertyPreset.cpp
new file mode 100644
--- /dev/null
+++ b/Engine/Plugins/Experimental/PropertyAnimatorCore/Source/PropertyAnimatorCore/Private/Presets/PropertyAnimatorCorePropertyPreset.cpp
@@ -0,0 +1,246 @@
+// Copyright Epic Games, Inc. All Rights Reserved.
+
+#include "Presets/PropertyAnimatorCorePropertyPreset.h"
+
+#include "Animators/PropertyAnimatorCoreBase.h"
+#include "Dom/JsonObject.h"
+#include "Dom/JsonValue.h"
+#include "Properties/PropertyAnimatorCoreContext.h"
+#include "Serialization/JsonReader.h"
+#include "Serialization/JsonSerializer.h"
+#include "Serialization/JsonWriter.h"
+
+void UPropertyAnimatorCorePropertyPreset::GetPresetProperties(const AActor* InActor, const UPropertyAnimatorCoreBase* InAnimator, TSet<FPropertyAnimatorCoreData>& OutProperties) const
+{
+	for (const TPair<FString, TSharedRef<FJsonValue>>& PropertyPreset : PropertyPresets)
+	{
+		FPropertyAnimatorCoreData Property(const_cast<AActor*>(InActor), PropertyPreset.Key);
+
+		if (Property.IsResolved())
+		{
+			OutProperties.Add(Property);
+		}
+	}
+}
+
+void UPropertyAnimatorCorePropertyPreset::OnPresetApplied(UPropertyAnimatorCoreBase* InAnimator, const TSet<FPropertyAnimatorCoreData>& InProperties)
+{
+	for (const TPair<FString, TSharedRef<FJsonValue>>& PropertyPreset : PropertyPresets)
+	{
+		FPropertyAnimatorCoreData Property(InAnimator->GetAnimatorActor(), PropertyPreset.Key);
+
+		if (!Property.IsResolved())
+		{
+			continue;
+		}
+
+		bool bFound = InProperties.Contains(Property);
+
+		if (!bFound)
+		{
+			for (const FPropertyAnimatorCoreData& LinkedProperty : InProperties)
+			{
+				if (LinkedProperty.IsChildOf(Property))
+				{
+					Property = LinkedProperty;
+					bFound = true;
+					break;
+				}
+			}
+		}
+
+		if (bFound)
+		{
+			if (UPropertyAnimatorCoreContext* Context = InAnimator->GetLinkedPropertyContext(Property))
+			{
+				Context->ImportPreset(this, PropertyPreset.Value);
+			}
+		}
+	}
+}
+
+void UPropertyAnimatorCorePropertyPreset::GetSupportedPresetProperties(const AActor* InActor, const UPropertyAnimatorCoreBase* InAnimator, TSet<FPropertyAnimatorCoreData>& OutProperties) const
+{
+	TSet<FPropertyAnimatorCoreData> PresetProperties;
+	GetPresetProperties(InActor, InAnimator, PresetProperties);
+	OutProperties.Empty(PresetProperties.Num());
+
+	if (PresetProperties.IsEmpty())
+	{
+		return;
+	}
+
+	for (const FPropertyAnimatorCoreData& PresetProperty : PresetProperties)
+	{
+		InAnimator->GetPropertiesSupported(PresetProperty, OutProperties, /** SearchDepth */3);
+	}
+}
+
+bool UPropertyAnimatorCorePropertyPreset::IsPresetSupported(const AActor* InActor, const UPropertyAnimatorCoreBase* InAnimator) const
+{
+	if (!IsValid(InActor) || !IsValid(InAnimator))
+	{
+		return false;
+	}
+
+	TSet<FPropertyAnimatorCoreData> SupportedProperties;
+	GetSupportedPresetProperties(InActor, InAnimator, SupportedProperties);
+
+	return !SupportedProperties.IsEmpty();
+}
+
+bool UPropertyAnimatorCorePropertyPreset::ApplyPreset(UPropertyAnimatorCoreBase* InAnimator)
+{
+	if (!IsValid(InAnimator) || InAnimator->IsTemplate())
+	{
+		return false;
+	}
+
+	TSet<FPropertyAnimatorCoreData> SupportedProperties;
+	GetSupportedPresetProperties(InAnimator->GetAnimatorActor(), InAnimator, SupportedProperties);
+
+	if (SupportedProperties.IsEmpty())
+	{
+		return false;
+	}
+
+	for (FPropertyAnimatorCoreData& SupportedProperty : SupportedProperties)
+	{
+		InAnimator->LinkProperty(SupportedProperty);
+	}
+
+	InAnimator->SetAnimatorDisplayName(FName(InAnimator->GetAnimatorOriginalName().ToString() + TEXT("_") + GetPresetDisplayName()));
+
+	OnPresetApplied(InAnimator, SupportedProperties);
+
+	return true;
+}
+
+bool UPropertyAnimatorCorePropertyPreset::IsPresetApplied(const UPropertyAnimatorCoreBase* InAnimator) const
+{
+	if (!IsValid(InAnimator) || InAnimator->IsTemplate())
+	{
+		return false;
+	}
+
+	TSet<FPropertyAnimatorCoreData> SupportedProperties;
+	GetSupportedPresetProperties(InAnimator->GetAnimatorActor(), InAnimator, SupportedProperties);
+
+	if (SupportedProperties.IsEmpty())
+	{
+		return false;
+	}
+
+	return InAnimator->IsPropertiesLinked(SupportedProperties);
+}
+
+bool UPropertyAnimatorCorePropertyPreset::UnapplyPreset(UPropertyAnimatorCoreBase* InAnimator)
+{
+	if (!IsValid(InAnimator) || InAnimator->IsTemplate())
+	{
+		return false;
+	}
+
+	TSet<FPropertyAnimatorCoreData> SupportedProperties;
+	GetSupportedPresetProperties(InAnimator->GetAnimatorActor(), InAnimator, SupportedProperties);
+
+	if (SupportedProperties.IsEmpty())
+	{
+		return false;
+	}
+
+	for (FPropertyAnimatorCoreData& SupportedProperty : SupportedProperties)
+	{
+		InAnimator->UnlinkProperty(SupportedProperty);
+	}
+
+	OnPresetUnapplied(InAnimator, SupportedProperties);
+
+	return true;
+}
+
+void UPropertyAnimatorCorePropertyPreset::CreatePreset(FName InName, const TArray<IPropertyAnimatorCorePresetable*>& InPresetableItems)
+{
+	Super::CreatePreset(InName, InPresetableItems);
+
+	TArray<TSharedPtr<FJsonValue>> JsonValues;
+
+	for (IPropertyAnimatorCorePresetable* InPresetableItem : InPresetableItems)
+	{
+		TSharedPtr<FJsonValue> JsonValue;
+
+		if (InPresetableItem
+			&& InPresetableItem->ExportPreset(this, JsonValue)
+			&& JsonValue.IsValid())
+		{
+			JsonValues.Add(JsonValue);
+		}
+	}
+
+	FString JsonString;
+	TSharedRef<TJsonWriter<>> JsonWriter = TJsonWriterFactory<>::Create(&JsonString);
+	if (FJsonSerializer::Serialize(JsonValues, JsonWriter))
+	{
+		PresetVersion = 0;
+		PresetContent = JsonString;
+	}
+}
+
+bool UPropertyAnimatorCorePropertyPreset::LoadPreset()
+{
+	if (PresetContent.IsEmpty())
+	{
+		return false;
+	}
+
+	TSharedRef<TJsonReader<>> Reader = TJsonReaderFactory<>::Create(PresetContent);
+
+	TArray<TSharedPtr<FJsonValue>> JsonArray;
+	if (!FJsonSerializer::Deserialize(Reader, JsonArray) || JsonArray.IsEmpty())
+	{
+		return false;
+	}
+
+	for (const TSharedPtr<FJsonValue>& JsonValue : JsonArray)
+	{
+		TSharedPtr<FJsonObject>* JsonObject = nullptr;
+		if (!JsonValue
+			|| !JsonValue->TryGetObject(JsonObject)
+			|| !JsonObject)
+		{
+			continue;
+		}
+
+		FString PropertyPath;
+
+		if ((*JsonObject)->TryGetStringField(UPropertyAnimatorCoreContext::GetAnimatedPropertyName().ToString(), PropertyPath)
+			&& !PropertyPath.IsEmpty())
+		{
+			PropertyPresets.Add(PropertyPath, JsonValue.ToSharedRef());
+		}
+	}
+
+	return !PropertyPresets.IsEmpty();
+}
+
+void UPropertyAnimatorCorePropertyPreset::GetAppliedPresetProperties(const UPropertyAnimatorCoreBase* InAnimator, TSet<FPropertyAnimatorCoreData>& OutSupportedProperties, TSet<FPropertyAnimatorCoreData>& OutAppliedProperties)
+{
+	OutSupportedProperties.Empty();
+	OutAppliedProperties.Empty();
+
+	if (!IsValid(InAnimator) || InAnimator->IsTemplate())
+	{
+		return;
+	}
+
+	GetSupportedPresetProperties(InAnimator->GetAnimatorActor(), InAnimator, OutSupportedProperties);
+	OutAppliedProperties.Reserve(OutSupportedProperties.Num());
+
+	for (const FPropertyAnimatorCoreData& Property : OutSupportedProperties)
+	{
+		if (InAnimator->IsPropertyLinked(Property))
+		{
+			OutAppliedProperties.Add(Property);
+		}
+	}
+}
diff --git a/Engine/Plugins/Experimental/PropertyAnimatorCore/Source/PropertyAnimatorCore/Private/Properties/PropertyAnimatorCoreContext.cpp b/Engine/Plugins/Experimental/PropertyAnimatorCore/Source/PropertyAnimatorCore/Private/Properties/PropertyAnimatorCoreContext.cpp
--- a/Engine/Plugins/Experimental/PropertyAnimatorCore/Source/PropertyAnimatorCore/Private/Properties/PropertyAnimatorCoreContext.cpp
+++ b/Engine/Plugins/Experimental/PropertyAnimatorCore/Source/PropertyAnimatorCore/Private/Properties/PropertyAnimatorCoreContext.cpp
@@ -3,6 +3,8 @@
 #include "Properties/PropertyAnimatorCoreContext.h"
 
 #include "Containers/Ticker.h"
+#include "Dom/JsonObject.h"
+#include "Dom/JsonValue.h"
 #include "StructUtils/InstancedStruct.h"
 #include "Properties/PropertyAnimatorCoreGroupBase.h"
 #include "Properties/PropertyAnimatorCoreResolver.h"
@@ -33,6 +35,11 @@ TArray<FPropertyAnimatorCoreData> UPropertyAnimatorCoreContext::ResolveProperty(
 	return ResolvedProperties;
 }
 
+FName UPropertyAnimatorCoreContext::GetAnimatedPropertyName()
+{
+	return GET_MEMBER_NAME_STRING_CHECKED(UPropertyAnimatorCoreContext, AnimatedProperty);
+}
+
 UPropertyAnimatorCoreBase* UPropertyAnimatorCoreContext::GetAnimator() const
 {
 	return GetTypedOuter<UPropertyAnimatorCoreBase>();
@@ -92,6 +99,11 @@ void UPropertyAnimatorCoreContext::SetMagnitude(float InMagnitude)
 	Magnitude = FMath::Clamp(InMagnitude, 0.f, 1.f);
 }
 
+void UPropertyAnimatorCoreContext::SetTimeOffset(double InOffset)
+{
+	TimeOffset = InOffset;
+}
+
 void UPropertyAnimatorCoreContext::SetMode(EPropertyAnimatorCoreMode InMode)
 {
 	if (InMode == Mode)
@@ -190,6 +202,62 @@ void UPropertyAnimatorCoreContext::PostEditChangeProperty(FPropertyChangedEvent&
 }
 #endif
 
+bool UPropertyAnimatorCoreContext::ImportPreset(const UPropertyAnimatorCorePresetBase* InPreset, const TSharedRef<FJsonValue>& InValue)
+{
+	TSharedPtr<FJsonObject>* JsonObject;
+	if (!InValue->TryGetObject(JsonObject) || !JsonObject || !JsonObject->IsValid())
+	{
+		return false;
+	}
+
+	bool bJsonAnimated = bAnimated;
+	(*JsonObject)->TryGetBoolField(GET_MEMBER_NAME_STRING_CHECKED(UPropertyAnimatorCoreContext, bAnimated), bJsonAnimated);
+	SetAnimated(bJsonAnimated);
+
+	if (bEditMagnitude)
+	{
+		float JsonMagnitude = Magnitude;
+		(*JsonObject)->TryGetNumberField(GET_MEMBER_NAME_STRING_CHECKED(UPropertyAnimatorCoreContext, Magnitude), JsonMagnitude);
+		SetMagnitude(JsonMagnitude);
+
+		float JsonTimeOffset = TimeOffset;
+		(*JsonObject)->TryGetNumberField(GET_MEMBER_NAME_STRING_CHECKED(UPropertyAnimatorCoreContext, TimeOffset), JsonTimeOffset);
+		SetTimeOffset(JsonTimeOffset);
+	}
+
+	if (bEditMode)
+	{
+		float JsonMode = static_cast<float>(Mode);
+		(*JsonObject)->TryGetNumberField(GET_MEMBER_NAME_STRING_CHECKED(UPropertyAnimatorCoreContext, Mode), JsonMode);
+		SetMode(static_cast<EPropertyAnimatorCoreMode>(JsonMode));
+	}
+
+	return true;
+}
+
+bool UPropertyAnimatorCoreContext::ExportPreset(const UPropertyAnimatorCorePresetBase* InPreset, TSharedPtr<FJsonValue>& OutValue)
+{
+	TSharedRef<FJsonObject> JsonObject = MakeShared<FJsonObject>();
+	OutValue = MakeShared<FJsonValueObject>(JsonObject);
+
+	JsonObject->SetBoolField(GET_MEMBER_NAME_STRING_CHECKED(UPropertyAnimatorCoreContext, bAnimated), bAnimated);
+
+	if (bEditMagnitude)
+	{
+		JsonObject->SetNumberField(GET_MEMBER_NAME_STRING_CHECKED(UPropertyAnimatorCoreContext, Magnitude), Magnitude);
+		JsonObject->SetNumberField(GET_MEMBER_NAME_STRING_CHECKED(UPropertyAnimatorCoreContext, TimeOffset), TimeOffset);
+	}
+
+	if (bEditMode)
+	{
+		JsonObject->SetNumberField(GET_MEMBER_NAME_STRING_CHECKED(UPropertyAnimatorCoreContext, Mode), static_cast<float>(Mode));
+	}
+
+	JsonObject->SetStringField(GET_MEMBER_NAME_STRING_CHECKED(UPropertyAnimatorCoreContext, AnimatedProperty), AnimatedProperty.GetPropertyLocatorPath());
+
+	return true;
+}
+
 void UPropertyAnimatorCoreContext::OnAnimatedPropertyLinked()
 {
 	bEditMagnitude = AnimatedProperty.IsA<FNumericProperty>() || AnimatedProperty.HasA<FNumericProperty>();
diff --git a/Engine/Plugins/Experimental/PropertyAnimatorCore/Source/PropertyAnimatorCore/Private/Properties/PropertyAnimatorCoreData.cpp b/Engine/Plugins/Experimental/PropertyAnimatorCore/Source/PropertyAnimatorCore/Private/Properties/PropertyAnimatorCoreData.cpp
--- a/Engine/Plugins/Experimental/PropertyAnimatorCore/Source/PropertyAnimatorCore/Private/Properties/PropertyAnimatorCoreData.cpp
+++ b/Engine/Plugins/Experimental/PropertyAnimatorCore/Source/PropertyAnimatorCore/Private/Properties/PropertyAnimatorCoreData.cpp
@@ -69,6 +69,108 @@ FPropertyAnimatorCoreData::FPropertyAnimatorCoreData(UObject* InObject, const TA
 	FindSetterFunctions();
 }
 
+FPropertyAnimatorCoreData::FPropertyAnimatorCoreData(AActor* InActor, const FString& InPropertyLocatorPath)
+{
+	UPropertyAnimatorCoreSubsystem* Subsystem = UPropertyAnimatorCoreSubsystem::Get();
+
+	if (!InActor || InPropertyLocatorPath.IsEmpty() || !Subsystem)
+	{
+		return;
+	}
+
+	UPropertyAnimatorCoreResolver* Resolver = nullptr;
+
+	TArray<FString> Elements;
+	InPropertyLocatorPath.ParseIntoArray(Elements, TEXT(";"));
+
+	if (Elements.Num() < 2)
+	{
+		return;
+	}
+
+	int32 ElementIndex = 0;
+
+	// Get Resolver
+	if (Elements.Num() == 3)
+	{
+		FString ResolverClass;
+		FString ResolverName;
+		Elements[ElementIndex].Split(TEXT(":"), &ResolverClass, &ResolverName);
+
+		Resolver = Subsystem->FindResolverByName(FName(ResolverName));
+		Resolver = Resolver ? Resolver : Subsystem->FindResolverByClass(LoadObject<UClass>(nullptr, *ResolverClass));
+
+		if (Resolver)
+		{
+			PropertyResolverClass = Resolver->GetClass();
+		}
+
+		ElementIndex++;
+	}
+
+	// Locate Owner
+	TArray<FString> Outers;
+	Elements[ElementIndex].ParseIntoArray(Outers, TEXT(","));
+
+	UObject* FoundObject = InActor;
+	for (const FString& Outer : Outers)
+	{
+		FString OuterClass;
+		FString OuterName;
+		Outer.Split(TEXT(":"), &OuterClass, &OuterName);
+
+		TArray<UObject*> OwnedObjects;
+		GetObjectsWithOuter(FoundObject, OwnedObjects);
+
+		bool bFound = false;
+
+		// Search by name
+		for (UObject* OwnedObject : OwnedObjects)
+		{
+			if (OwnedObject->GetName() == OuterName)
+			{
+				FoundObject = OwnedObject;
+				bFound = true;
+				break;
+			}
+		}
+
+		// Search by class
+		if (!bFound)
+		{
+			for (UObject* OwnedObject : OwnedObjects)
+			{
+				if (OwnedObject->GetClass()->GetClassPathName().ToString() == OuterClass)
+				{
+					FoundObject = OwnedObject;
+					bFound = true;
+					break;
+				}
+			}
+		}
+
+		if (!bFound)
+		{
+			return;
+		}
+	}
+
+	OwnerWeak = FoundObject;
+	ElementIndex++;
+
+	// Locate Property
+	TArray<FString> Properties;
+	Elements[ElementIndex].ParseIntoArray(Properties, TEXT(","));
+
+	for (const FString& Property : Properties)
+	{
+		ChainProperties.Add(FindFProperty<FProperty>(*Property));
+	}
+
+	GeneratePropertyPath();
+	FindSetterFunctions();
+}
+
 bool FPropertyAnimatorCoreData::IsResolvable() const
 {
 	return PropertyResolverClass.Get() != nullptr;
@@ -237,6 +339,37 @@ FString FPropertyAnimatorCoreData::GetPropertyDisplayName() const
 	return PropertyDisplayName;
 }
 
+FString FPropertyAnimatorCoreData::GetPropertyLocatorPath() const
+{
+	FString LocatorPath;
+
+	// Append resolver
+	if (const UPropertyAnimatorCoreResolver* Resolver = GetPropertyResolver())
+	{
+		LocatorPath += Resolver->GetClass()->GetClassPathName().ToString() + TEXT(":") + Resolver->GetResolverName().ToString() + TEXT(";");
+	}
+
+	// Append outers
+	for (const UObject* Outer : GetOuters(GetOwningActor()))
+	{
+		LocatorPath += Outer->GetClass()->GetClassPathName().ToString() + TEXT(":") + Outer->GetName() + TEXT(",");
+	}
+
+	LocatorPath.RemoveFromEnd(TEXT(","));
+	LocatorPath += TEXT(";");
+
+	// Append properties
+	const AActor* OwningActor = GetOwningActor();
+	for (const TFieldPath<FProperty>& Property : ChainProperties)
+	{
+		LocatorPath += Property->GetPathName(OwningActor) + TEXT(",");
+	}
+
+	LocatorPath.RemoveFromEnd(TEXT(","));
+
+	return LocatorPath;
+}
+
 FName FPropertyAnimatorCoreData::GetMemberPropertyName() const
 {
 	const FProperty* MemberProperty = GetMemberProperty();
diff --git a/Engine/Plugins/Experimental/PropertyAnimatorCore/Source/PropertyAnimatorCore/Private/Subsystems/PropertyAnimatorCoreSubsystem.cpp b/Engine/Plugins/Experimental/PropertyAnimatorCore/Source/PropertyAnimatorCore/Private/Subsystems/PropertyAnimatorCoreSubsystem.cpp
--- a/Engine/Plugins/Experimental/PropertyAnimatorCore/Source/PropertyAnimatorCore/Private/Subsystems/PropertyAnimatorCoreSubsystem.cpp
+++ b/Engine/Plugins/Experimental/PropertyAnimatorCore/Source/PropertyAnimatorCore/Private/Subsystems/PropertyAnimatorCoreSubsystem.cpp
@@ -1,7 +1,10 @@
 // Copyright Epic Games, Inc. All Rights Reserved.
 
 #include "Subsystems/PropertyAnimatorCoreSubsystem.h"
+
 #include "Animators/PropertyAnimatorCoreBase.h"
+#include "AssetRegistry/AssetData.h"
+#include "AssetRegistry/AssetRegistryModule.h"
 #include "Components/PropertyAnimatorCoreComponent.h"
 #include "Engine/Engine.h"
 #include "GameFramework/Actor.h"
@@ -49,6 +52,14 @@ void UPropertyAnimatorCoreSubsystem::Initialize(FSubsystemCollectionBase& Collec
 	RegisterPropertyAlias(PropertyType + GET_MEMBER_NAME_STRING_CHECKED(FRotator, Roll), TEXT("X"));
 	RegisterPropertyAlias(PropertyType + GET_MEMBER_NAME_STRING_CHECKED(FRotator, Pitch), TEXT("Y"));
 	RegisterPropertyAlias(PropertyType + GET_MEMBER_NAME_STRING_CHECKED(FRotator, Yaw), TEXT("Z"));
+
+	const FAssetRegistryModule& AssetRegistryModule = FModuleManager::LoadModuleChecked<FAssetRegistryModule>("AssetRegistry");
+
+	IAssetRegistry& AssetRegistry = AssetRegistryModule.Get();
+	AssetRegistry.OnFilesLoaded().AddUObject(this, &UPropertyAnimatorCoreSubsystem::OnAssetRegistryFilesLoaded);
+	AssetRegistry.OnAssetAdded().AddUObject(this, &UPropertyAnimatorCoreSubsystem::OnAssetRegistryAssetAdded);
+	AssetRegistry.OnAssetRemoved().AddUObject(this, &UPropertyAnimatorCoreSubsystem::OnAssetRegistryAssetRemoved);
+	AssetRegistry.OnAssetUpdated().AddUObject(this, &UPropertyAnimatorCoreSubsystem::OnAssetRegistryAssetUpdated);
 }
 
 void UPropertyAnimatorCoreSubsystem::Deinitialize()
@@ -60,6 +71,15 @@ void UPropertyAnimatorCoreSubsystem::Deinitialize()
 	PresetsWeak.Empty();
 	SetterResolvers.Empty();
 
+	if (const FAssetRegistryModule* AssetRegistryModule = FModuleManager::GetModulePtr<FAssetRegistryModule>("AssetRegistry"))
+	{
+		IAssetRegistry& AssetRegistry = AssetRegistryModule->Get();
+		AssetRegistry.OnFilesLoaded().RemoveAll(this);
+		AssetRegistry.OnAssetAdded().RemoveAll(this);
+		AssetRegistry.OnAssetRemoved().RemoveAll(this);
+		AssetRegistry.OnAssetUpdated().RemoveAll(this);
+	}
+
 	Super::Deinitialize();
 }
 
@@ -617,7 +637,7 @@ bool UPropertyAnimatorCoreSubsystem::LinkAnimatorProperties(UPropertyAnimatorCor
 		bool bResult = false;
 		for (const FPropertyAnimatorCoreData& PropertyData : InProperties)
 		{
-			bResult |= InAnimator->LinkProperty(PropertyData);
+			bResult |= (InAnimator->LinkProperty(PropertyData) != nullptr);
 		}
 
 		return bResult;
@@ -791,6 +811,36 @@ bool UPropertyAnimatorCoreSubsystem::UnregisterResolverClass(const UClass* InRes
 	return false;
 }
 
+UPropertyAnimatorCoreResolver* UPropertyAnimatorCoreSubsystem::FindResolverByName(FName InResolverName)
+{
+	for (const TWeakObjectPtr<UPropertyAnimatorCoreResolver>& ResolverWeakPair : ResolversWeak)
+	{
+		UPropertyAnimatorCoreResolver* Resolver = ResolverWeakPair.Get();
+
+		if (Resolver && Resolver->GetResolverName().IsEqual(InResolverName))
+		{
+			return Resolver;
+		}
+	}
+
+	return nullptr;
+}
+
+UPropertyAnimatorCoreResolver* UPropertyAnimatorCoreSubsystem::FindResolverByClass(const UClass* InResolverClass)
+{
+	for (const TWeakObjectPtr<UPropertyAnimatorCoreResolver>& ResolverWeakPair : ResolversWeak)
+	{
+		UPropertyAnimatorCoreResolver* Resolver = ResolverWeakPair.Get();
+
+		if (Resolver && Resolver->GetClass() == InResolverClass)
+		{
+			return Resolver;
+		}
+	}
+
+	return nullptr;
+}
+
 bool UPropertyAnimatorCoreSubsystem::IsResolverClassRegistered(const UClass* InResolverClass) const
 {
 	for (const TWeakObjectPtr<UPropertyAnimatorCoreResolver>& ResolverWeakPair : ResolversWeak)
@@ -988,7 +1038,11 @@ bool UPropertyAnimatorCoreSubsystem::RegisterPresetClass(const UClass* InPresetC
 
 	if (UPropertyAnimatorCorePresetBase* CDO = InPresetClass->GetDefaultObject<UPropertyAnimatorCorePresetBase>())
 	{
-		PresetsWeak.Add(CDO);
+		if (CDO->LoadPreset())
+		{
+			PresetsWeak.Add(CDO);
+			CDO->OnPresetRegistered();
+		}
 
 		return true;
 	}
@@ -1033,7 +1087,28 @@ bool UPropertyAnimatorCoreSubsystem::IsPresetClassRegistered(const UClass* InPre
 	return false;
 }
 
-TSet<UPropertyAnimatorCorePresetBase*> UPropertyAnimatorCoreSubsystem::GetSupportedPresets(const AActor* InActor, const UPropertyAnimatorCoreBase* InAnimator) const
+TSet<UPropertyAnimatorCorePresetBase*> UPropertyAnimatorCoreSubsystem::GetAvailablePresets(TSubclassOf<UPropertyAnimatorCorePresetBase> InPresetClass) const
+{
+	TSet<UPropertyAnimatorCorePresetBase*> AvailablePresets;
+	AvailablePresets.Reserve(PresetsWeak.Num());
+
+	Algo::TransformIf(
+		PresetsWeak
+		, AvailablePresets
+		, [&InPresetClass](const TWeakObjectPtr<UPropertyAnimatorCorePresetBase>& InPresetWeak)
+		{
+			return InPresetWeak.IsValid() && InPresetWeak->IsA(InPresetClass);
+		}
+		, [](const TWeakObjectPtr<UPropertyAnimatorCorePresetBase>& InPresetWeak)
+		{
+			return InPresetWeak.Get();
+		}
+	);
+
+	return AvailablePresets;
+}
+
+TSet<UPropertyAnimatorCorePresetBase*> UPropertyAnimatorCoreSubsystem::GetSupportedPresets(const AActor* InActor, const UPropertyAnimatorCoreBase* InAnimator, TSubclassOf<UPropertyAnimatorCorePresetBase> InPresetClass) const
 {
 	TSet<UPropertyAnimatorCorePresetBase*> SupportedPresets;
 
@@ -1041,7 +1116,7 @@ TSet<UPropertyAnimatorCorePresetBase*> UPropertyAnimatorCoreSubsystem::GetSuppor
 	{
 		UPropertyAnimatorCorePresetBase* Preset = PresetWeak.Get();
 
-		if (!Preset)
+		if (!Preset || !Preset->IsA(InPresetClass))
 		{
 			continue;
 		}
@@ -1057,27 +1132,6 @@ TSet<UPropertyAnimatorCorePresetBase*> UPropertyAnimatorCoreSubsystem::GetSuppor
 	return SupportedPresets;
 }
 
-TSet<UPropertyAnimatorCorePresetBase*> UPropertyAnimatorCoreSubsystem::GetAvailablePresets() const
-{
-	TSet<UPropertyAnimatorCorePresetBase*> AvailablePresets;
-	AvailablePresets.Reserve(PresetsWeak.Num());
-
-	Algo::TransformIf(
-		PresetsWeak
-		, AvailablePresets
-		, [](const TWeakObjectPtr<UPropertyAnimatorCorePresetBase>& InPresetWeak)
-		{
-			return InPresetWeak.IsValid();
-		}
-		, [](const TWeakObjectPtr<UPropertyAnimatorCorePresetBase>& InPresetWeak)
-		{
-			return InPresetWeak.Get();
-		}
-	);
-
-	return AvailablePresets;
-}
-
 bool UPropertyAnimatorCoreSubsystem::RegisterConverterClass(const UClass* InConverterClass)
 {
 	if (!IsValid(InConverterClass))
@@ -1366,4 +1420,77 @@ void UPropertyAnimatorCoreSubsystem::RegisterAnimatorClasses()
 	}
 }
 
+void UPropertyAnimatorCoreSubsystem::OnAssetRegistryFilesLoaded()
+{
+	bFilesLoaded = true;
+
+	const FAssetRegistryModule& AssetRegistryModule = FModuleManager::LoadModuleChecked<FAssetRegistryModule>("AssetRegistry");
+
+	TArray<FAssetData> Assets;
+	AssetRegistryModule.Get().GetAssetsByClass(UPropertyAnimatorCorePresetBase::StaticClass()->GetClassPathName(), Assets, /** Subclass */true);
+
+	for (const FAssetData& Asset : Assets)
+	{
+		RegisterPresetAsset(Asset);
+	}
+}
+
+void UPropertyAnimatorCoreSubsystem::OnAssetRegistryAssetAdded(const FAssetData& InAssetData)
+{
+	if (bFilesLoaded)
+	{
+		RegisterPresetAsset(InAssetData);
+	}
+}
+
+void UPropertyAnimatorCoreSubsystem::OnAssetRegistryAssetRemoved(const FAssetData& InAssetData)
+{
+	UnregisterPresetAsset(InAssetData);
+}
+
+void UPropertyAnimatorCoreSubsystem::OnAssetRegistryAssetUpdated(const FAssetData& InAssetData)
+{
+	UnregisterPresetAsset(InAssetData);
+	RegisterPresetAsset(InAssetData);
+}
+
+void UPropertyAnimatorCoreSubsystem::RegisterPresetAsset(const FAssetData& InAssetData)
+{
+	if (const UClass* Class = InAssetData.GetClass(EResolveClass::Yes))
+	{
+		if (!Class->IsChildOf<UPropertyAnimatorCorePresetBase>())
+		{
+			return;
+		}
+
+		if (UPropertyAnimatorCorePresetBase* Preset = Cast<UPropertyAnimatorCorePresetBase>(InAssetData.GetAsset()))
+		{
+			if (Preset->LoadPreset())
+			{
+				PresetsWeak.Add(Preset);
+				Preset->OnPresetRegistered();
+			}
+		}
+	}
+}
+
+void UPropertyAnimatorCoreSubsystem::UnregisterPresetAsset(const FAssetData& InAssetData)
+{
+	if (const UClass* Class = InAssetData.GetClass(EResolveClass::Yes))
+	{
+		if (!Class->IsChildOf<UPropertyAnimatorCorePresetBase>())
+		{
+			return;
+		}
+
+		if (UPropertyAnimatorCorePresetBase* Preset = Cast<UPropertyAnimatorCorePresetBase>(InAssetData.GetAsset()))
+		{
+			if (PresetsWeak.Remove(Preset) > 0)
+			{
+				Preset->OnPresetUnregistered();
+			}
+		}
+	}
+}
+
 #undef LOCTEXT_NAMESPACE
diff --git a/Engine/Plugins/Experimental/PropertyAnimatorCore/Source/PropertyAnimatorCore/Private/TimeSources/PropertyAnimatorCoreManualTimeSource.cpp b/Engine/Plugins/Experimental/PropertyAnimatorCore/Source/PropertyAnimatorCore/Private/TimeSources/PropertyAnimatorCoreManualTimeSource.cpp
--- a/Engine/Plugins/Experimental/PropertyAnimatorCore/Source/PropertyAnimatorCore/Private/TimeSources/PropertyAnimatorCoreManualTimeSource.cpp
+++ b/Engine/Plugins/Experimental/PropertyAnimatorCore/Source/PropertyAnimatorCore/Private/TimeSources/PropertyAnimatorCoreManualTimeSource.cpp
@@ -2,6 +2,8 @@
 
 #include "TimeSources/PropertyAnimatorCoreManualTimeSource.h"
 
+#include "Dom/JsonObject.h"
+#include "Dom/JsonValue.h"
 #include "Misc/App.h"
 
 void UPropertyAnimatorCoreManualTimeSource::SetOverrideTime(bool bInOverride)
@@ -77,6 +79,41 @@ void UPropertyAnimatorCoreManualTimeSource::OnTimeSourceInactive()
 	Stop();
 }
 
+bool UPropertyAnimatorCoreManualTimeSource::ImportPreset(const UPropertyAnimatorCorePresetBase* InPreset, const TSharedRef<FJsonValue>& InValue)
+{
+	const TSharedPtr<FJsonObject>* JsonTimeSourceObject;
+
+	if (Super::ImportPreset(InPreset, InValue) && InValue->TryGetObject(JsonTimeSourceObject))
+	{
+		double JsonCustomTime = CustomTime;
+		(*JsonTimeSourceObject)->TryGetNumberField(GET_MEMBER_NAME_STRING_CHECKED(UPropertyAnimatorCoreManualTimeSource, CustomTime), JsonCustomTime);
+		SetCustomTime(JsonCustomTime);
+
+		bool bJsonOverrideTime = bOverrideTime;
+		(*JsonTimeSourceObject)->TryGetBoolField(GET_MEMBER_NAME_STRING_CHECKED(UPropertyAnimatorCoreManualTimeSource, bOverrideTime), bJsonOverrideTime);
+		SetOverrideTime(bJsonOverrideTime);
+
+		return true;
+	}
+
+	return false;
+}
+
+bool UPropertyAnimatorCoreManualTimeSource::ExportPreset(const UPropertyAnimatorCorePresetBase* InPreset, TSharedPtr<FJsonValue>& OutValue)
+{
+	const TSharedPtr<FJsonObject>* JsonTimeSourceObject;
+
+	if (Super::ExportPreset(InPreset, OutValue) && OutValue->TryGetObject(JsonTimeSourceObject))
+	{
+		(*JsonTimeSourceObject)->SetNumberField(GET_MEMBER_NAME_STRING_CHECKED(UPropertyAnimatorCoreManualTimeSource, CustomTime), CustomTime);
+		(*JsonTimeSourceObject)->SetBoolField(GET_MEMBER_NAME_STRING_CHECKED(UPropertyAnimatorCoreManualTimeSource, bOverrideTime), bOverrideTime);
+
+		return true;
+	}
+
+	return false;
+}
+
 void UPropertyAnimatorCoreManualTimeSource::Play(bool bInForward)
 {
 	if (bOverrideTime)
@@ -110,7 +147,6 @@ void UPropertyAnimatorCoreManualTimeSource::Stop()
 		return;
 	}
 
-	CustomTime = 0.f;
 	Pause();
 	ActiveStatus = EPropertyAnimatorCoreManualStatus::Stopped;
 }
diff --git a/Engine/Plugins/Experimental/PropertyAnimatorCore/Source/PropertyAnimatorCore/Private/TimeSources/PropertyAnimatorCoreSequencerTimeSource.cpp b/Engine/Plugins/Experimental/PropertyAnimatorCore/Source/PropertyAnimatorCore/Private/TimeSources/PropertyAnimatorCoreSequencerTimeSource.cpp
--- a/Engine/Plugins/Experimental/PropertyAnimatorCore/Source/PropertyAnimatorCore/Private/TimeSources/PropertyAnimatorCoreSequencerTimeSource.cpp
+++ b/Engine/Plugins/Experimental/PropertyAnimatorCore/Source/PropertyAnimatorCore/Private/TimeSources/PropertyAnimatorCoreSequencerTimeSource.cpp
@@ -18,6 +18,10 @@ bool UPropertyAnimatorCoreSequencerTimeSource::UpdateEvaluationData(FPropertyAni
 void UPropertyAnimatorCoreSequencerTimeSource::OnSequencerTimeEvaluated(const TOptional<double>& InTimeEval, const TOptional<float>& InMagnitudeEval)
 {
 	EvalResult.bEvalValid = InTimeEval.IsSet() && InMagnitudeEval.IsSet();
-	EvalResult.EvalTime = InTimeEval.Get(0.0);
-	EvalResult.EvalMagnitude = InMagnitudeEval.Get(1.f);
+
+	if (EvalResult.bEvalValid)
+	{
+		EvalResult.EvalTime = InTimeEval.GetValue();
+		EvalResult.EvalMagnitude = InMagnitudeEval.GetValue();
+	}
 }
diff --git a/Engine/Plugins/Experimental/PropertyAnimatorCore/Source/PropertyAnimatorCore/Private/TimeSources/PropertyAnimatorCoreSystemTimeSource.cpp b/Engine/Plugins/Experimental/PropertyAnimatorCore/Source/PropertyAnimatorCore/Private/TimeSources/PropertyAnimatorCoreSystemTimeSource.cpp
--- a/Engine/Plugins/Experimental/PropertyAnimatorCore/Source/PropertyAnimatorCore/Private/TimeSources/PropertyAnimatorCoreSystemTimeSource.cpp
+++ b/Engine/Plugins/Experimental/PropertyAnimatorCore/Source/PropertyAnimatorCore/Private/TimeSources/PropertyAnimatorCoreSystemTimeSource.cpp
@@ -2,6 +2,8 @@
 
 #include "TimeSources/PropertyAnimatorCoreSystemTimeSource.h"
 
+#include "Dom/JsonObject.h"
+#include "Dom/JsonValue.h"
 #include "Internationalization/Regex.h"
 #include "Misc/DateTime.h"
 
@@ -42,6 +44,41 @@ void UPropertyAnimatorCoreSystemTimeSource::OnTimeSourceActive()
 	OnModeChanged();
 }
 
+bool UPropertyAnimatorCoreSystemTimeSource::ImportPreset(const UPropertyAnimatorCorePresetBase* InPreset, const TSharedRef<FJsonValue>& InValue)
+{
+	const TSharedPtr<FJsonObject>* JsonTimeSourceObject;
+
+	if (Super::ImportPreset(InPreset, InValue) && InValue->TryGetObject(JsonTimeSourceObject))
+	{
+		FString JsonCountdown = CountdownDuration;
+		(*JsonTimeSourceObject)->TryGetStringField(GET_MEMBER_NAME_STRING_CHECKED(UPropertyAnimatorCoreSystemTimeSource, CountdownDuration), JsonCountdown);
+		SetCountdownDuration(JsonCountdown);
+
+		double JsonMode = static_cast<double>(Mode);
+		(*JsonTimeSourceObject)->TryGetNumberField(GET_MEMBER_NAME_STRING_CHECKED(UPropertyAnimatorCoreSystemTimeSource, Mode), JsonMode);
+		SetMode(static_cast<EPropertyAnimatorCoreSystemMode>(JsonMode));
+
+		return true;
+	}
+
+	return false;
+}
+
+bool UPropertyAnimatorCoreSystemTimeSource::ExportPreset(const UPropertyAnimatorCorePresetBase* InPreset, TSharedPtr<FJsonValue>& OutValue)
+{
+	TSharedPtr<FJsonObject>* JsonTimeSourceObject = nullptr;
+
+	if (Super::ExportPreset(InPreset, OutValue) && OutValue->TryGetObject(JsonTimeSourceObject))
+	{
+		(*JsonTimeSourceObject)->SetNumberField(GET_MEMBER_NAME_STRING_CHECKED(UPropertyAnimatorCoreSystemTimeSource, Mode), static_cast<double>(Mode));
+		(*JsonTimeSourceObject)->SetStringField(GET_MEMBER_NAME_STRING_CHECKED(UPropertyAnimatorCoreSystemTimeSource, CountdownDuration), CountdownDuration);
+
+		return true;
+	}
+
+	return false;
+}
+
 void UPropertyAnimatorCoreSystemTimeSource::SetMode(EPropertyAnimatorCoreSystemMode InMode)
 {
 	if (Mode == InMode)
diff --git a/Engine/Plugins/Experimental/PropertyAnimatorCore/Source/PropertyAnimatorCore/Private/TimeSources/PropertyAnimatorCoreTimeSourceBase.cpp b/Engine/Plugins/Experimental/PropertyAnimatorCore/Source/PropertyAnimatorCore/Private/TimeSources/PropertyAnimatorCoreTimeSourceBase.cpp
--- a/Engine/Plugins/Experimental/PropertyAnimatorCore/Source/PropertyAnimatorCore/Private/TimeSources/PropertyAnimatorCoreTimeSourceBase.cpp
+++ b/Engine/Plugins/Experimental/PropertyAnimatorCore/Source/PropertyAnimatorCore/Private/TimeSources/PropertyAnimatorCoreTimeSourceBase.cpp
@@ -2,6 +2,9 @@
 
 #include "TimeSources/PropertyAnimatorCoreTimeSourceBase.h"
 
+#include "Dom/JsonObject.h"
+#include "Dom/JsonValue.h"
+
 void UPropertyAnimatorCoreTimeSourceBase::ActivateTimeSource()
 {
 	if (IsTimeSourceActive())
@@ -29,7 +32,7 @@ EPropertyAnimatorCoreTimeSourceResult UPropertyAnimatorCoreTimeSourceBase::Fetch
 	if (!UpdateEvaluationData(OutEvaluationData))
 	{
 		// Reset evaluation state
-		return EPropertyAnimatorCoreTimeSourceResult::Reset;
+		return EPropertyAnimatorCoreTimeSourceResult::Idle;
 	}
 
 	if (!IsFramerateAllowed(OutEvaluationData.TimeElapsed))
@@ -53,6 +56,36 @@ void UPropertyAnimatorCoreTimeSourceBase::SetUseFrameRate(bool bInUseFrameRate)
 	bUseFrameRate = bInUseFrameRate;
 }
 
+bool UPropertyAnimatorCoreTimeSourceBase::ImportPreset(const UPropertyAnimatorCorePresetBase* InPreset, const TSharedRef<FJsonValue>& InValue)
+{
+	const TSharedPtr<FJsonObject>* JsonTimeSourceObject;
+	if (!InValue->TryGetObject(JsonTimeSourceObject) || !JsonTimeSourceObject)
+	{
+		return false;
+	}
+
+	bool bJsonUseFrameRate = bUseFrameRate;
+	(*JsonTimeSourceObject)->TryGetBoolField(GET_MEMBER_NAME_STRING_CHECKED(UPropertyAnimatorCoreTimeSourceBase, bUseFrameRate), bJsonUseFrameRate);
+	SetUseFrameRate(bJsonUseFrameRate);
+
+	double JsonFrameRate = FrameRate;
+	(*JsonTimeSourceObject)->TryGetNumberField(GET_MEMBER_NAME_STRING_CHECKED(UPropertyAnimatorCoreTimeSourceBase, FrameRate), JsonFrameRate);
+	SetFrameRate(JsonFrameRate);
+
+	return true;
+}
+
+bool UPropertyAnimatorCoreTimeSourceBase::ExportPreset(const UPropertyAnimatorCorePresetBase* InPreset, TSharedPtr<FJsonValue>& OutValue)
+{
+	TSharedRef<FJsonObject> JsonTimeSourceObject = MakeShared<FJsonObject>();
+	OutValue = MakeShared<FJsonValueObject>(JsonTimeSourceObject);
+
+	JsonTimeSourceObject->SetBoolField(GET_MEMBER_NAME_STRING_CHECKED(UPropertyAnimatorCoreTimeSourceBase, bUseFrameRate), bUseFrameRate);
+	JsonTimeSourceObject->SetNumberField(GET_MEMBER_NAME_STRING_CHECKED(UPropertyAnimatorCoreTimeSourceBase, FrameRate), FrameRate);
+
+	return true;
+}
+
 bool UPropertyAnimatorCoreTimeSourceBase::UpdateEvaluationData(FPropertyAnimatorCoreTimeSourceEvaluationData& OutData)
 {
 	return false;
diff --git a/Engine/Plugins/Experimental/PropertyAnimatorCore/Source/PropertyAnimatorCore/PropertyAnimatorCore.Build.cs b/Engine/Plugins/Experimental/PropertyAnimatorCore/Source/PropertyAnimatorCore/PropertyAnimatorCore.Build.cs
--- a/Engine/Plugins/Experimental/PropertyAnimatorCore/Source/PropertyAnimatorCore/PropertyAnimatorCore.Build.cs
+++ b/Engine/Plugins/Experimental/PropertyAnimatorCore/Source/PropertyAnimatorCore/PropertyAnimatorCore.Build.cs
@@ -23,7 +23,8 @@ public class PropertyAnimatorCore : ModuleRules
 				"ApplicationCore",
 				"CoreUObject",
 				"DeveloperSettings",
-				"Engine"
+				"Engine",
+				"Json"
 			}
 		);
 
diff --git a/Engine/Plugins/Experimental/PropertyAnimatorCore/Source/PropertyAnimatorCore/Public/Animators/PropertyAnimatorCoreBase.h b/Engine/Plugins/Experimental/PropertyAnimatorCore/Source/PropertyAnimatorCore/Public/Animators/PropertyAnimatorCoreBase.h
--- a/Engine/Plugins/Experimental/PropertyAnimatorCore/Source/PropertyAnimatorCore/Public/Animators/PropertyAnimatorCoreBase.h
+++ b/Engine/Plugins/Experimental/PropertyAnimatorCore/Source/PropertyAnimatorCore/Public/Animators/PropertyAnimatorCoreBase.h
@@ -2,6 +2,7 @@
 
 #pragma once
 
+#include "Presets/PropertyAnimatorCorePresetable.h"
 #include "Properties/PropertyAnimatorCoreContext.h"
 #include "Properties/PropertyAnimatorCoreData.h"
 #include "UObject/Object.h"
@@ -28,7 +29,7 @@ struct FPropertyAnimatorCoreMetadata
 
 /** Abstract base class for any Animator, holds a set of linked properties */
 UCLASS(MinimalAPI, Abstract, EditInlineNew, AutoExpandCategories=("Animator"))
-class UPropertyAnimatorCoreBase : public UObject
+class UPropertyAnimatorCoreBase : public UObject, public IPropertyAnimatorCorePresetable
 {
 	GENERATED_BODY()
 
@@ -123,7 +124,7 @@ public:
 	PROPERTYANIMATORCORE_API int32 GetLinkedPropertiesCount() const;
 
 	/** Link property to this Animator to be able to drive it */
-	PROPERTYANIMATORCORE_API bool LinkProperty(const FPropertyAnimatorCoreData& InLinkProperty);
+	PROPERTYANIMATORCORE_API UPropertyAnimatorCoreContext* LinkProperty(const FPropertyAnimatorCoreData& InLinkProperty);
 
 	/** Unlink property from this Animator */
 	PROPERTYANIMATORCORE_API bool UnlinkProperty(const FPropertyAnimatorCoreData& InUnlinkProperty);
@@ -164,6 +165,11 @@ public:
 	/** Get the context for the linked property */
 	PROPERTYANIMATORCORE_API UPropertyAnimatorCoreContext* GetLinkedPropertyContext(const FPropertyAnimatorCoreData& InProperty) const;
 
+	TConstArrayView<UPropertyAnimatorCoreContext*> GetLinkedPropertiesContext() const
+	{
+		return LinkedProperties;
+	}
+
 	/** Get the casted context for the linked property */
 	template<typename InContextClass
 		UE_REQUIRES(TIsDerivedFrom<InContextClass, UPropertyAnimatorCoreContext>::Value)>
@@ -209,6 +215,11 @@ public:
 		return true;
 	}
 
+	//~ Begin IPropertyAnimatorCorePresetable
+	PROPERTYANIMATORCORE_API virtual bool ImportPreset(const UPropertyAnimatorCorePresetBase* InPreset, const TSharedRef<FJsonValue>& InValue) override;
+	PROPERTYANIMATORCORE_API virtual bool ExportPreset(const UPropertyAnimatorCorePresetBase* InPreset, TSharedPtr<FJsonValue>& OutValue) override;
+	//~ End IPropertyAnimatorCorePresetable
+
 protected:
 	//~ Begin UObject
 	PROPERTYANIMATORCORE_API virtual void PostCDOContruct() override;
@@ -323,6 +334,9 @@ private:
 
 	void OnTimeSourceNameChanged();
 
+	/** Called when time source enters idle/invalid state */
+	void OnTimeSourceEnterIdleState();
+
 	/** Called after an action that causes the owner to change */
 	void ResolvePropertiesOwner(AActor* InNewOwner = nullptr);
 
diff --git a/Engine/Plugins/Experimental/PropertyAnimatorCore/Source/PropertyAnimatorCore/Public/Presets/PropertyAnimatorCoreAnimatorPreset.h b/Engine/Plugins/Experimental/PropertyAnimatorCore/Source/PropertyAnimatorCore/Public/Presets/PropertyAnimatorCoreAnimatorPreset.h
new file mode 100644
--- /dev/null
+++ b/Engine/Plugins/Experimental/PropertyAnimatorCore/Source/PropertyAnimatorCore/Public/Presets/PropertyAnimatorCoreAnimatorPreset.h
@@ -0,0 +1,38 @@
+// Copyright Epic Games, Inc. All Rights Reserved.
+
+#pragma once
+
+#include "PropertyAnimatorCorePresetBase.h"
+#include "Templates/SharedPointer.h"
+#include "Templates/SubclassOf.h"
+#include "PropertyAnimatorCoreAnimatorPreset.generated.h"
+
+class FJsonValue;
+
+/**
+ * Animator preset class used to import/export animator data
+ */
+UCLASS(MinimalAPI, BlueprintType)
+class UPropertyAnimatorCoreAnimatorPreset : public UPropertyAnimatorCorePresetBase
+{
+	GENERATED_BODY()
+
+public:
+	//~ Begin UPropertyAnimatorCorePresetBase
+	virtual bool IsPresetApplied(const UPropertyAnimatorCoreBase* InAnimator) const override;
+	virtual bool IsPresetSupported(const AActor* InActor, const UPropertyAnimatorCoreBase* InAnimator) const override;
+	virtual bool ApplyPreset(UPropertyAnimatorCoreBase* InAnimator) override;
+	virtual bool UnapplyPreset(UPropertyAnimatorCoreBase* InAnimator) override;
+	virtual void CreatePreset(FName InName, const TArray<IPropertyAnimatorCorePresetable*>& InPresetableItem) override;
+	virtual bool LoadPreset() override;
+	//~ End UPropertyAnimatorCorePresetBase
+
+	PROPERTYANIMATORCORE_API UPropertyAnimatorCoreBase* GetAnimatorTemplate() const;
+
+protected:
+	/** Animator class to target for this preset */
+	UPROPERTY(Transient)
+	TSubclassOf<UPropertyAnimatorCoreBase> TargetAnimatorClass;
+
+	TSharedPtr<FJsonValue> AnimatorPreset;
+};
diff --git a/Engine/Plugins/Experimental/PropertyAnimatorCore/Source/PropertyAnimatorCore/Public/Presets/PropertyAnimatorCorePresetBase.h b/Engine/Plugins/Experimental/PropertyAnimatorCore/Source/PropertyAnimatorCore/Public/Presets/PropertyAnimatorCorePresetBase.h
--- a/Engine/Plugins/Experimental/PropertyAnimatorCore/Source/PropertyAnimatorCore/Public/Presets/PropertyAnimatorCorePresetBase.h
+++ b/Engine/Plugins/Experimental/PropertyAnimatorCore/Source/PropertyAnimatorCore/Public/Presets/PropertyAnimatorCorePresetBase.h
@@ -2,10 +2,10 @@
 
 #pragma once
 
-#include "Properties/PropertyAnimatorCoreData.h"
 #include "PropertyAnimatorCorePresetBase.generated.h"
 
 class AActor;
+class IPropertyAnimatorCorePresetable;
 class UPropertyAnimatorCoreBase;
 
 /**
@@ -13,7 +13,7 @@ class UPropertyAnimatorCoreBase;
  * Will get registered automatically by the subsystem
  * Should remain transient and stateless
  */
-UCLASS(MinimalAPI, Abstract, Transient)
+UCLASS(MinimalAPI, BlueprintType, Abstract)
 class UPropertyAnimatorCorePresetBase : public UObject
 {
 	GENERATED_BODY()
@@ -34,32 +34,43 @@ public:
 
 	PROPERTYANIMATORCORE_API FString GetPresetDisplayName() const;
 
-	/** Get the preset properties for that actor */
-	virtual void GetPresetProperties(const AActor* InActor, const UPropertyAnimatorCoreBase* InAnimator, TSet<FPropertyAnimatorCoreData>& OutProperties) const {}
-
-	/** Get the preset properties for that actor but only supported ones for that animator */
-	PROPERTYANIMATORCORE_API virtual void GetSupportedPresetProperties(const AActor* InActor, const UPropertyAnimatorCoreBase* InAnimator, TSet<FPropertyAnimatorCoreData>& OutProperties) const;
-
 	/** Checks if the preset is supported on that actor and animator */
-	PROPERTYANIMATORCORE_API virtual bool IsPresetSupported(const AActor* InActor, const UPropertyAnimatorCoreBase* InAnimator) const;
+	PROPERTYANIMATORCORE_API virtual bool IsPresetSupported(const AActor* InActor, const UPropertyAnimatorCoreBase* InAnimator) const PURE_VIRTUAL(UPropertyAnimatorCorePresetBase::IsPresetSupported, return false; );
 
 	/** Checks if the preset is applied to this animator */
-	PROPERTYANIMATORCORE_API virtual bool IsPresetApplied(const UPropertyAnimatorCoreBase* InAnimator) const;
+	PROPERTYANIMATORCORE_API virtual bool IsPresetApplied(const UPropertyAnimatorCoreBase* InAnimator) const PURE_VIRTUAL(UPropertyAnimatorCorePresetBase::IsPresetApplied, return false; );
 
 	/** Applies this preset on the newly created animator */
-	PROPERTYANIMATORCORE_API virtual bool ApplyPreset(UPropertyAnimatorCoreBase* InAnimator);
+	PROPERTYANIMATORCORE_API virtual bool ApplyPreset(UPropertyAnimatorCoreBase* InAnimator) PURE_VIRTUAL(UPropertyAnimatorCorePresetBase::ApplyPreset, return false; );
 
 	/** Un applies this preset on an animator */
-	PROPERTYANIMATORCORE_API virtual bool UnapplyPreset(UPropertyAnimatorCoreBase* InAnimator);
+	PROPERTYANIMATORCORE_API virtual bool UnapplyPreset(UPropertyAnimatorCoreBase* InAnimator) PURE_VIRTUAL(UPropertyAnimatorCorePresetBase::UnapplyPreset, return false; );
+
+	/** Called once to create this preset out of supported items */
+	PROPERTYANIMATORCORE_API virtual void CreatePreset(FName InName, const TArray<IPropertyAnimatorCorePresetable*>& InPresetableItem);
+
+	/** Called once before registering preset to load it, if load fails then preset is not registered */
+	virtual bool LoadPreset()
+	{
+		return true;
+	}
 
-	/** Gets the supported and applied properties for an animator */
-	PROPERTYANIMATORCORE_API virtual void GetAppliedPresetProperties(const UPropertyAnimatorCoreBase* InAnimator, TSet<FPropertyAnimatorCoreData>& OutSupportedProperties, TSet<FPropertyAnimatorCoreData>& OutAppliedProperties);
+	/** Called when this preset is registered by the subsystem */
+	virtual void OnPresetRegistered() {}
 
-	/** Called when this preset is applied on the animator */
-	virtual void OnPresetApplied(UPropertyAnimatorCoreBase* InAnimator, const TSet<FPropertyAnimatorCoreData>& InProperties) {}
+	/** Called when this preset is unregistered by the subsystem */
+	virtual void OnPresetUnregistered() {}
 
-private:
+protected:
 	/** Name used to display this preset to the user */
-	UPROPERTY(Transient)
+	UPROPERTY(VisibleInstanceOnly, Category="Animator")
 	FName PresetName;
+
+	/** Version of this preset for diffs */
+	UPROPERTY(VisibleInstanceOnly, Category="Animator")
+	int32 PresetVersion = INDEX_NONE;
+
+	/** Preset stringify content */
+	UPROPERTY(VisibleInstanceOnly, Category="Animator")
+	FString PresetContent;
 };
\ No newline at end of file
diff --git a/Engine/Plugins/Experimental/PropertyAnimatorCore/Source/PropertyAnimatorCore/Public/Presets/PropertyAnimatorCorePresetable.h b/Engine/Plugins/Experimental/PropertyAnimatorCore/Source/PropertyAnimatorCore/Public/Presets/PropertyAnimatorCorePresetable.h
new file mode 100644
--- /dev/null
+++ b/Engine/Plugins/Experimental/PropertyAnimatorCore/Source/PropertyAnimatorCore/Public/Presets/PropertyAnimatorCorePresetable.h
@@ -0,0 +1,29 @@
+// Copyright Epic Games, Inc. All Rights Reserved.
+
+#pragma once
+
+#include "UObject/Interface.h"
+#include "Templates/SharedPointer.h"
+#include "PropertyAnimatorCorePresetable.generated.h"
+
+class FJsonValue;
+class UPropertyAnimatorCorePresetBase;
+
+UINTERFACE(MinimalAPI)
+class UPropertyAnimatorCorePresetable : public UInterface
+{
+	GENERATED_BODY()
+};
+
+/** Preset interface for animator system */
+class IPropertyAnimatorCorePresetable
+{
+	GENERATED_BODY()
+
+public:
+	/** Import a specific preset */
+	PROPERTYANIMATORCORE_API virtual bool ImportPreset(const UPropertyAnimatorCorePresetBase* InPreset, const TSharedRef<FJsonValue>& InValue) = 0;
+
+	/** Export a specific preset */
+	PROPERTYANIMATORCORE_API virtual bool ExportPreset(const UPropertyAnimatorCorePresetBase* InPreset, TSharedPtr<FJsonValue>& OutValue) = 0;
+};
\ No newline at end of file
diff --git a/Engine/Plugins/Experimental/PropertyAnimatorCore/Source/PropertyAnimatorCore/Public/Presets/PropertyAnimatorCorePropertyPreset.h b/Engine/Plugins/Experimental/PropertyAnimatorCore/Source/PropertyAnimatorCore/Public/Presets/PropertyAnimatorCorePropertyPreset.h
new file mode 100644
--- /dev/null
+++ b/Engine/Plugins/Experimental/PropertyAnimatorCore/Source/PropertyAnimatorCore/Public/Presets/PropertyAnimatorCorePropertyPreset.h
@@ -0,0 +1,48 @@
+// Copyright Epic Games, Inc. All Rights Reserved.
+
+#pragma once
+
+#include "Containers/Map.h"
+#include "Properties/PropertyAnimatorCoreData.h"
+#include "PropertyAnimatorCorePresetBase.h"
+#include "Templates/SharedPointer.h"
+#include "PropertyAnimatorCorePropertyPreset.generated.h"
+
+class FJsonValue;
+
+/**
+ * Property preset class used to import/export properties on supported animators
+ */
+UCLASS(MinimalAPI, BlueprintType)
+class UPropertyAnimatorCorePropertyPreset : public UPropertyAnimatorCorePresetBase
+{
+	GENERATED_BODY()
+
+public:
+	//~ Begin UPropertyAnimatorCorePresetBase
+	PROPERTYANIMATORCORE_API virtual bool IsPresetSupported(const AActor* InActor, const UPropertyAnimatorCoreBase* InAnimator) const override;
+	PROPERTYANIMATORCORE_API virtual bool IsPresetApplied(const UPropertyAnimatorCoreBase* InAnimator) const override;
+	PROPERTYANIMATORCORE_API virtual bool ApplyPreset(UPropertyAnimatorCoreBase* InAnimator) override;
+	PROPERTYANIMATORCORE_API virtual bool UnapplyPreset(UPropertyAnimatorCoreBase* InAnimator) override;
+	PROPERTYANIMATORCORE_API virtual void CreatePreset(FName InName, const TArray<IPropertyAnimatorCorePresetable*>& InPresetableItems) override;
+	PROPERTYANIMATORCORE_API virtual bool LoadPreset() override;
+	//~ End UPropertyAnimatorCorePresetBase
+
+	/** Called when this preset is applied on the animator */
+	PROPERTYANIMATORCORE_API virtual void OnPresetApplied(UPropertyAnimatorCoreBase* InAnimator, const TSet<FPropertyAnimatorCoreData>& InProperties);
+
+	/** Called when this preset is unapplied on the animator */
+	virtual void OnPresetUnapplied(UPropertyAnimatorCoreBase* InAnimator, const TSet<FPropertyAnimatorCoreData>& InProperties) {}
+
+	/** Get the preset properties for that actor */
+	PROPERTYANIMATORCORE_API virtual void GetPresetProperties(const AActor* InActor, const UPropertyAnimatorCoreBase* InAnimator, TSet<FPropertyAnimatorCoreData>& OutProperties) const;
+
+	/** Get the preset properties for that actor but only supported ones for that animator */
+	PROPERTYANIMATORCORE_API virtual void GetSupportedPresetProperties(const AActor* InActor, const UPropertyAnimatorCoreBase* InAnimator, TSet<FPropertyAnimatorCoreData>& OutProperties) const;
+
+	/** Gets the supported and applied properties for an animator */
+	PROPERTYANIMATORCORE_API virtual void GetAppliedPresetProperties(const UPropertyAnimatorCoreBase* InAnimator, TSet<FPropertyAnimatorCoreData>& OutSupportedProperties, TSet<FPropertyAnimatorCoreData>& OutAppliedProperties);
+
+protected:
+	TMap<FString, TSharedRef<FJsonValue>> PropertyPresets;
+};
diff --git a/Engine/Plugins/Experimental/PropertyAnimatorCore/Source/PropertyAnimatorCore/Public/Properties/PropertyAnimatorCoreContext.h b/Engine/Plugins/Experimental/PropertyAnimatorCore/Source/PropertyAnimatorCore/Public/Properties/PropertyAnimatorCoreContext.h
--- a/Engine/Plugins/Experimental/PropertyAnimatorCore/Source/PropertyAnimatorCore/Public/Properties/PropertyAnimatorCoreContext.h
+++ b/Engine/Plugins/Experimental/PropertyAnimatorCore/Source/PropertyAnimatorCore/Public/Properties/PropertyAnimatorCoreContext.h
@@ -3,6 +3,8 @@
 #pragma once
 
 #include "PropertyAnimatorCoreData.h"
+#include "Presets/PropertyAnimatorCorePresetable.h"
+#include "Presets/PropertyAnimatorCorePresetBase.h"
 #include "StructUtils/PropertyBag.h"
 #include "UObject/Object.h"
 #include "PropertyAnimatorCoreContext.generated.h"
@@ -17,13 +19,15 @@ class UScriptStruct;
 UENUM(BlueprintType)
 enum class EPropertyAnimatorCoreMode : uint8
 {
+	/** Set the property value directly */
 	Absolute,
-	Additive,
+	/** Add value on the existing property value */
+	Additive
 };
 
 /** Context for properties linked to an animator */
 UCLASS(MinimalAPI, BlueprintType)
-class UPropertyAnimatorCoreContext : public UObject
+class UPropertyAnimatorCoreContext : public UObject, public IPropertyAnimatorCorePresetable
 {
 	GENERATED_BODY()
 
@@ -31,6 +35,8 @@ class UPropertyAnimatorCoreContext : public UObject
 	friend class FPropertyAnimatorCoreEditorContextTypeCustomization;
 
 public:
+	static FName GetAnimatedPropertyName();
+
 	const FPropertyAnimatorCoreData& GetAnimatedProperty() const
 	{
 		return AnimatedProperty;
@@ -59,6 +65,12 @@ public:
 		return Magnitude;
 	}
 
+	PROPERTYANIMATORCORE_API void SetTimeOffset(double InOffset);
+	double GetTimeOffset() const
+	{
+		return TimeOffset;
+	}
+
 	PROPERTYANIMATORCORE_API void SetMode(EPropertyAnimatorCoreMode InMode);
 	EPropertyAnimatorCoreMode GetMode() const
 	{
@@ -94,6 +106,11 @@ public:
 		return false;
 	}
 
+	//~ Begin IPropertyAnimatorCorePresetable
+	PROPERTYANIMATORCORE_API virtual bool ImportPreset(const UPropertyAnimatorCorePresetBase* InPreset, const TSharedRef<FJsonValue>& InValue) override;
+	PROPERTYANIMATORCORE_API virtual bool ExportPreset(const UPropertyAnimatorCorePresetBase* InPreset, TSharedPtr<FJsonValue>& OutValue) override;
+	//~ End IPropertyAnimatorCorePresetable
+
 protected:
 	//~ Begin UObject
 	PROPERTYANIMATORCORE_API virtual void PostLoad() override;
@@ -155,6 +172,10 @@ private:
 	UPROPERTY(EditInstanceOnly, Setter, Getter, Category="Animator", meta=(ClampMin="0.0", ClampMax="1.0", HideEditConditionToggle, EditCondition="bEditMagnitude", EditConditionHides, AllowPrivateAccess="true"))
 	float Magnitude = 1.f;
 
+	/** Time offset to evaluate this property */
+	UPROPERTY(EditInstanceOnly, Setter, Getter, Category="Animator", meta=(Units=Seconds, HideEditConditionToggle, EditCondition="bEditMagnitude", EditConditionHides, AllowPrivateAccess="true"))
+	double TimeOffset = 0;
+
 	/** Edit condition for modes */
 	UPROPERTY(Transient)
 	bool bEditMode = true;
diff --git a/Engine/Plugins/Experimental/PropertyAnimatorCore/Source/PropertyAnimatorCore/Public/Properties/PropertyAnimatorCoreData.h b/Engine/Plugins/Experimental/PropertyAnimatorCore/Source/PropertyAnimatorCore/Public/Properties/PropertyAnimatorCoreData.h
--- a/Engine/Plugins/Experimental/PropertyAnimatorCore/Source/PropertyAnimatorCore/Public/Properties/PropertyAnimatorCoreData.h
+++ b/Engine/Plugins/Experimental/PropertyAnimatorCore/Source/PropertyAnimatorCore/Public/Properties/PropertyAnimatorCoreData.h
@@ -35,6 +35,9 @@ struct FPropertyAnimatorCoreData
 	/** Take the owner, the property chain until the inner property and lastly the inner property */
 	PROPERTYANIMATORCORE_API explicit FPropertyAnimatorCoreData(UObject* InObject, const TArray<FProperty*>& InChainProperties, FProperty* InProperty, TSubclassOf<UPropertyAnimatorCoreResolver> InResolverClass = nullptr);
 
+	/** Takes an actor and a locator path, tries to resolve property and owner */
+	PROPERTYANIMATORCORE_API explicit FPropertyAnimatorCoreData(AActor* InActor, const FString& InPropertyLocatorPath);
+
 	/** Is this a resolvable property that uses a custom resolver */
 	PROPERTYANIMATORCORE_API bool IsResolvable() const;
 
@@ -103,6 +106,9 @@ struct FPropertyAnimatorCoreData
 		return PathHash;
 	}
 
+	/** Creates a single path that contains resolver, owner and properties for (de)serialization */
+	FString GetPropertyLocatorPath() const;
+
 	/** The member property name */
 	PROPERTYANIMATORCORE_API FName GetMemberPropertyName() const;
 
diff --git a/Engine/Plugins/Experimental/PropertyAnimatorCore/Source/PropertyAnimatorCore/Public/Subsystems/PropertyAnimatorCoreSubsystem.h b/Engine/Plugins/Experimental/PropertyAnimatorCore/Source/PropertyAnimatorCore/Public/Subsystems/PropertyAnimatorCoreSubsystem.h
--- a/Engine/Plugins/Experimental/PropertyAnimatorCore/Source/PropertyAnimatorCore/Public/Subsystems/PropertyAnimatorCoreSubsystem.h
+++ b/Engine/Plugins/Experimental/PropertyAnimatorCore/Source/PropertyAnimatorCore/Public/Subsystems/PropertyAnimatorCoreSubsystem.h
@@ -78,6 +78,10 @@ public:
 	/** Unregister a resolver */
 	PROPERTYANIMATORCORE_API bool UnregisterResolverClass(const UClass* InResolverClass);
 
+	UPropertyAnimatorCoreResolver* FindResolverByName(FName InResolverName);
+
+	UPropertyAnimatorCoreResolver* FindResolverByClass(const UClass* InResolverClass);
+
 	/** Is this resolver registered */
 	PROPERTYANIMATORCORE_API bool IsResolverClassRegistered(const UClass* InResolverClass) const;
 
@@ -114,11 +118,11 @@ public:
 	/** Is this preset class registered */
 	PROPERTYANIMATORCORE_API bool IsPresetClassRegistered(const UClass* InPresetClass) const;
 
-	/** Gets all supported presets for a specific animator and actor */
-	PROPERTYANIMATORCORE_API TSet<UPropertyAnimatorCorePresetBase*> GetSupportedPresets(const AActor* InActor, const UPropertyAnimatorCoreBase* InAnimator) const;
-
 	/** Get all registered preset available */
-	PROPERTYANIMATORCORE_API TSet<UPropertyAnimatorCorePresetBase*> GetAvailablePresets() const;
+	PROPERTYANIMATORCORE_API TSet<UPropertyAnimatorCorePresetBase*> GetAvailablePresets(TSubclassOf<UPropertyAnimatorCorePresetBase> InPresetClass) const;
+
+	/** Gets all supported presets for a specific animator and actor */
+	PROPERTYANIMATORCORE_API TSet<UPropertyAnimatorCorePresetBase*> GetSupportedPresets(const AActor* InActor, const UPropertyAnimatorCoreBase* InAnimator, TSubclassOf<UPropertyAnimatorCorePresetBase> InPresetClass) const;
 
 	PROPERTYANIMATORCORE_API bool RegisterSetterResolver(FName InPropertyName, TFunction<UFunction*(const UObject*)>&& InFunction);
 
@@ -209,6 +213,13 @@ protected:
 	 */
 	void RegisterAnimatorClasses();
 
+	void OnAssetRegistryFilesLoaded();
+	void OnAssetRegistryAssetAdded(const FAssetData& InAssetData);
+	void OnAssetRegistryAssetRemoved(const FAssetData& InAssetData);
+	void OnAssetRegistryAssetUpdated(const FAssetData& InAssetData);
+	void RegisterPresetAsset(const FAssetData& InAssetData);
+	void UnregisterPresetAsset(const FAssetData& InAssetData);
+
 	/** Time sources available to use with animators */
 	UPROPERTY()
 	TSet<TWeakObjectPtr<UPropertyAnimatorCoreTimeSourceBase>> TimeSourcesWeak;
@@ -238,4 +249,6 @@ protected:
 
 	/** Some property should have a friendlier name and replace the original name by an alias */
 	TMap<FString, FString> PropertyAliases;
+
+	bool bFilesLoaded = false;
 };
diff --git a/Engine/Plugins/Experimental/PropertyAnimatorCore/Source/PropertyAnimatorCore/Public/TimeSources/PropertyAnimatorCoreManualTimeSource.h b/Engine/Plugins/Experimental/PropertyAnimatorCore/Source/PropertyAnimatorCore/Public/TimeSources/PropertyAnimatorCoreManualTimeSource.h
--- a/Engine/Plugins/Experimental/PropertyAnimatorCore/Source/PropertyAnimatorCore/Public/TimeSources/PropertyAnimatorCoreManualTimeSource.h
+++ b/Engine/Plugins/Experimental/PropertyAnimatorCore/Source/PropertyAnimatorCore/Public/TimeSources/PropertyAnimatorCoreManualTimeSource.h
@@ -56,12 +56,6 @@ public:
 		return State;
 	}
 
-	//~ Begin UPropertyAnimatorTimeSourceBase
-	virtual bool UpdateEvaluationData(FPropertyAnimatorCoreTimeSourceEvaluationData& OutData) override;
-	virtual void OnTimeSourceActive() override;
-	virtual void OnTimeSourceInactive() override;
-	//~ End UPropertyAnimatorTimeSourceBase
-
 	void Play(bool bInForward);
 	void Pause();
 	void Stop();
@@ -76,6 +70,14 @@ protected:
 #endif
 	//~ End UObject
 
+	//~ Begin UPropertyAnimatorTimeSourceBase
+	virtual bool UpdateEvaluationData(FPropertyAnimatorCoreTimeSourceEvaluationData& OutData) override;
+	virtual void OnTimeSourceActive() override;
+	virtual void OnTimeSourceInactive() override;
+	virtual bool ImportPreset(const UPropertyAnimatorCorePresetBase* InPreset, const TSharedRef<FJsonValue>& InValue) override;
+	virtual bool ExportPreset(const UPropertyAnimatorCorePresetBase* InPreset, TSharedPtr<FJsonValue>& OutValue) override;
+	//~ End UPropertyAnimatorTimeSourceBase
+
 	void OnStateChanged();
 	void OnOverrideTimeChanged();
 
diff --git a/Engine/Plugins/Experimental/PropertyAnimatorCore/Source/PropertyAnimatorCore/Public/TimeSources/PropertyAnimatorCoreSystemTimeSource.h b/Engine/Plugins/Experimental/PropertyAnimatorCore/Source/PropertyAnimatorCore/Public/TimeSources/PropertyAnimatorCoreSystemTimeSource.h
--- a/Engine/Plugins/Experimental/PropertyAnimatorCore/Source/PropertyAnimatorCore/Public/TimeSources/PropertyAnimatorCoreSystemTimeSource.h
+++ b/Engine/Plugins/Experimental/PropertyAnimatorCore/Source/PropertyAnimatorCore/Public/TimeSources/PropertyAnimatorCoreSystemTimeSource.h
@@ -32,11 +32,6 @@ public:
 		: UPropertyAnimatorCoreTimeSourceBase(TEXT("System"))
 	{}
 
-	//~ Begin UPropertyAnimatorTimeSourceBase
-	virtual bool UpdateEvaluationData(FPropertyAnimatorCoreTimeSourceEvaluationData& OutData) override;
-	virtual void OnTimeSourceActive() override;
-	//~ End UPropertyAnimatorTimeSourceBase
-
 	void SetMode(EPropertyAnimatorCoreSystemMode InMode);
 	EPropertyAnimatorCoreSystemMode GetMode() const
 	{
@@ -64,6 +59,13 @@ protected:
 #endif
 	//~ End UObject
 
+	//~ Begin UPropertyAnimatorTimeSourceBase
+	virtual bool UpdateEvaluationData(FPropertyAnimatorCoreTimeSourceEvaluationData& OutData) override;
+	virtual void OnTimeSourceActive() override;
+	virtual bool ImportPreset(const UPropertyAnimatorCorePresetBase* InPreset, const TSharedRef<FJsonValue>& InValue) override;
+	virtual bool ExportPreset(const UPropertyAnimatorCorePresetBase* InPreset, TSharedPtr<FJsonValue>& OutValue) override;
+	//~ End UPropertyAnimatorTimeSourceBase
+
 	void OnModeChanged();
 
 	/** Machine time mode to use */
diff --git a/Engine/Plugins/Experimental/PropertyAnimatorCore/Source/PropertyAnimatorCore/Public/TimeSources/PropertyAnimatorCoreTimeSourceBase.h b/Engine/Plugins/Experimental/PropertyAnimatorCore/Source/PropertyAnimatorCore/Public/TimeSources/PropertyAnimatorCoreTimeSourceBase.h
--- a/Engine/Plugins/Experimental/PropertyAnimatorCore/Source/PropertyAnimatorCore/Public/TimeSources/PropertyAnimatorCoreTimeSourceBase.h
+++ b/Engine/Plugins/Experimental/PropertyAnimatorCore/Source/PropertyAnimatorCore/Public/TimeSources/PropertyAnimatorCoreTimeSourceBase.h
@@ -2,6 +2,7 @@
 
 #pragma once
 
+#include "Presets/PropertyAnimatorCorePresetable.h"
 #include "PropertyAnimatorCoreTimeSourceBase.generated.h"
 
 class UPropertyAnimatorCoreBase;
@@ -11,8 +12,8 @@ enum class EPropertyAnimatorCoreTimeSourceResult
 {
 	/** Time is the same as previous, evaluation can be skipped */
 	Skip,
-	/** Time is in an invalid state or out of range, restore values */
-	Reset,
+	/** Time is in an invalid state or out of range */
+	Idle,
 	/** Time is valid and in range, evaluate time */
 	Evaluate
 };
@@ -32,7 +33,7 @@ struct FPropertyAnimatorCoreTimeSourceEvaluationData
  * Can be transient or saved to disk if contains user set data
  */
 UCLASS(MinimalAPI, Abstract)
-class UPropertyAnimatorCoreTimeSourceBase : public UObject
+class UPropertyAnimatorCoreTimeSourceBase : public UObject, public IPropertyAnimatorCorePresetable
 {
 	GENERATED_BODY()
 
@@ -79,6 +80,11 @@ public:
 		return LastTimeElapsed;
 	}
 
+	//~ Begin IPropertyAnimatorCorePresetable
+	PROPERTYANIMATORCORE_API virtual bool ImportPreset(const UPropertyAnimatorCorePresetBase* InPreset, const TSharedRef<FJsonValue>& InValue) override;
+	PROPERTYANIMATORCORE_API virtual bool ExportPreset(const UPropertyAnimatorCorePresetBase* InPreset, TSharedPtr<FJsonValue>& OutValue) override;
+	//~ End IPropertyAnimatorCorePresetable
+
 protected:
 	/** Retrieve evaluation data to provide animators, return true if data is valid, false otherwise */
 	PROPERTYANIMATORCORE_API virtual bool UpdateEvaluationData(FPropertyAnimatorCoreTimeSourceEvaluationData& OutData);
diff --git a/Engine/Plugins/Experimental/PropertyAnimatorCore/Source/PropertyAnimatorCoreEditor/Private/Customizations/PropertyAnimatorCoreEditorDetailCustomization.cpp b/Engine/Plugins/Experimental/PropertyAnimatorCore/Source/PropertyAnimatorCoreEditor/Private/Customizations/PropertyAnimatorCoreEditorDetailCustomization.cpp
--- a/Engine/Plugins/Experimental/PropertyAnimatorCore/Source/PropertyAnimatorCoreEditor/Private/Customizations/PropertyAnimatorCoreEditorDetailCustomization.cpp
+++ b/Engine/Plugins/Experimental/PropertyAnimatorCore/Source/PropertyAnimatorCoreEditor/Private/Customizations/PropertyAnimatorCoreEditorDetailCustomization.cpp
@@ -4,6 +4,8 @@
 
 #include "DetailLayoutBuilder.h"
 #include "DetailWidgetRow.h"
+#include "Presets/PropertyAnimatorCorePropertyPreset.h"
+#include "Styles/PropertyAnimatorCoreEditorStyle.h"
 #include "Subsystems/PropertyAnimatorCoreEditorSubsystem.h"
 #include "Subsystems/PropertyAnimatorCoreSubsystem.h"
 #include "ToolMenus.h"
@@ -80,6 +82,21 @@ void FPropertyAnimatorCoreEditorDetailCustomization::CustomizeDetails(IDetailLay
 				.DesiredSizeOverride(FVector2D(16.f))
 				.Image(FAppStyle::Get().GetBrush("Icons.Plus"))
 			]
+		]
+		+ SHorizontalBox::Slot()
+		.AutoWidth()
+		.Padding(2.f, 0.f)
+		[
+			SNew(SButton)
+			.ContentPadding(2.f)
+			.ButtonStyle(&FCoreStyle::Get().GetWidgetStyle<FButtonStyle>("Button"))
+			.ToolTipText(LOCTEXT("CreatePropertyPreset", "Create a preset from these properties"))
+			.OnClicked(this, &FPropertyAnimatorCoreEditorDetailCustomization::OnCreatePropertyPresetClicked)
+			[
+				SNew(SImage)
+				.DesiredSizeOverride(FVector2D(16.f))
+				.Image(FPropertyAnimatorCoreEditorStyle::Get().GetBrush("PropertyControlIcon.Export"))
+			]
 		];
 
 	if (FDetailWidgetDecl* ValueWidget = PropertyRow->CustomValueWidget())
@@ -243,4 +260,20 @@ FReply FPropertyAnimatorCoreEditorDetailCustomization::UnlinkProperties() const
 	return FReply::Handled();
 }
 
+FReply FPropertyAnimatorCoreEditorDetailCustomization::OnCreatePropertyPresetClicked() const
+{
+	UPropertyAnimatorCoreBase* Animator = AnimatorWeak.Get();
+	if (!Animator)
+	{
+		return FReply::Handled();
+	}
+
+	if (UPropertyAnimatorCoreEditorSubsystem* AnimatorEditorSubsystem = UPropertyAnimatorCoreEditorSubsystem::Get())
+	{
+		AnimatorEditorSubsystem->CreatePresetAsset(UPropertyAnimatorCorePropertyPreset::StaticClass(), TArray<IPropertyAnimatorCorePresetable*>{Animator->GetLinkedPropertiesContext()});
+	}
+
+	return FReply::Handled();
+}
+
 #undef LOCTEXT_NAMESPACE
diff --git a/Engine/Plugins/Experimental/PropertyAnimatorCore/Source/PropertyAnimatorCoreEditor/Private/Customizations/PropertyAnimatorCoreEditorDetailCustomization.h b/Engine/Plugins/Experimental/PropertyAnimatorCore/Source/PropertyAnimatorCoreEditor/Private/Customizations/PropertyAnimatorCoreEditorDetailCustomization.h
--- a/Engine/Plugins/Experimental/PropertyAnimatorCore/Source/PropertyAnimatorCoreEditor/Private/Customizations/PropertyAnimatorCoreEditorDetailCustomization.h
+++ b/Engine/Plugins/Experimental/PropertyAnimatorCore/Source/PropertyAnimatorCoreEditor/Private/Customizations/PropertyAnimatorCoreEditorDetailCustomization.h
@@ -36,5 +36,7 @@ protected:
 
 	FReply UnlinkProperties() const;
 
+	FReply OnCreatePropertyPresetClicked() const;
+
 	TWeakObjectPtr<UPropertyAnimatorCoreBase> AnimatorWeak;
 };
diff --git a/Engine/Plugins/Experimental/PropertyAnimatorCore/Source/PropertyAnimatorCoreEditor/Private/Customizations/PropertyAnimatorCoreEditorStackCustomization.cpp b/Engine/Plugins/Experimental/PropertyAnimatorCore/Source/PropertyAnimatorCoreEditor/Private/Customizations/PropertyAnimatorCoreEditorStackCustomization.cpp
--- a/Engine/Plugins/Experimental/PropertyAnimatorCore/Source/PropertyAnimatorCoreEditor/Private/Customizations/PropertyAnimatorCoreEditorStackCustomization.cpp
+++ b/Engine/Plugins/Experimental/PropertyAnimatorCore/Source/PropertyAnimatorCoreEditor/Private/Customizations/PropertyAnimatorCoreEditorStackCustomization.cpp
@@ -9,6 +9,8 @@
 #include "Items/OperatorStackEditorGroupItem.h"
 #include "Items/OperatorStackEditorObjectItem.h"
 #include "Menus/PropertyAnimatorCoreEditorMenu.h"
+#include "Presets/PropertyAnimatorCoreAnimatorPreset.h"
+#include "Styles/PropertyAnimatorCoreEditorStyle.h"
 #include "Styling/SlateIconFinder.h"
 #include "Subsystems/PropertyAnimatorCoreEditorSubsystem.h"
 #include "Subsystems/PropertyAnimatorCoreSubsystem.h"
@@ -357,6 +359,24 @@ void UPropertyAnimatorCoreEditorStackCustomization::RemoveAnimatorAction(FOperat
 	}
 }
 
+bool UPropertyAnimatorCoreEditorStackCustomization::CanExportAnimator(FOperatorStackEditorItemPtr InItem) const
+{
+	return InItem.IsValid() && InItem->GetValueCount() == 1 && InItem->HasValue(0);
+}
+
+void UPropertyAnimatorCoreEditorStackCustomization::ExportAnimatorAction(FOperatorStackEditorItemPtr InItem)
+{
+	if (!CanExportAnimator(InItem))
+	{
+		return;
+	}
+
+	if (UPropertyAnimatorCoreEditorSubsystem* AnimatorEditorSubsystem = UPropertyAnimatorCoreEditorSubsystem::Get())
+	{
+		AnimatorEditorSubsystem->CreatePresetAsset(UPropertyAnimatorCoreAnimatorPreset::StaticClass(), {InItem->Get<UPropertyAnimatorCoreBase>(0)});
+	}
+}
+
 void UPropertyAnimatorCoreEditorStackCustomization::FillAddAnimatorMenuSection(UToolMenu* InToolMenu) const
 {
 	if (!InToolMenu)
@@ -423,6 +443,23 @@ void UPropertyAnimatorCoreEditorStackCustomization::FillAnimatorHeaderActionMenu
 		return;
 	}
 
+	if (ItemContext->IsA<UPropertyAnimatorCoreBase>())
+	{
+		const FToolMenuEntry ExportAnimatorAction = FToolMenuEntry::InitToolBarButton(
+			TEXT("ExportAnimatorMenuEntry")
+			, FUIAction(
+				FExecuteAction::CreateUObject(this, &UPropertyAnimatorCoreEditorStackCustomization::ExportAnimatorAction, ItemContext)
+				, FCanExecuteAction::CreateUObject(this, &UPropertyAnimatorCoreEditorStackCustomization::CanExportAnimator, ItemContext)
+				, FIsActionChecked()
+				, FIsActionButtonVisible::CreateUObject(this, &UPropertyAnimatorCoreEditorStackCustomization::CanExportAnimator, ItemContext))
+			, FText::GetEmpty()
+			, FText::GetEmpty()
+			, FSlateIcon(FPropertyAnimatorCoreEditorStyle::Get().GetStyleSetName(), "PropertyControlIcon.Export")
+		);
+
+		InToolMenu->AddMenuEntry(ExportAnimatorAction.Name, ExportAnimatorAction);
+	}
+
 	const FToolMenuEntry RemoveAnimatorAction = FToolMenuEntry::InitToolBarButton(
 		TEXT("RemoveAnimatorMenuEntry")
 		, FExecuteAction::CreateUObject(this, &UPropertyAnimatorCoreEditorStackCustomization::RemoveAnimatorAction, ItemContext)
diff --git a/Engine/Plugins/Experimental/PropertyAnimatorCore/Source/PropertyAnimatorCoreEditor/Private/Customizations/PropertyAnimatorCoreEditorStackCustomization.h b/Engine/Plugins/Experimental/PropertyAnimatorCore/Source/PropertyAnimatorCoreEditor/Private/Customizations/PropertyAnimatorCoreEditorStackCustomization.h
--- a/Engine/Plugins/Experimental/PropertyAnimatorCore/Source/PropertyAnimatorCoreEditor/Private/Customizations/PropertyAnimatorCoreEditorStackCustomization.h
+++ b/Engine/Plugins/Experimental/PropertyAnimatorCore/Source/PropertyAnimatorCoreEditor/Private/Customizations/PropertyAnimatorCoreEditorStackCustomization.h
@@ -32,6 +32,9 @@ protected:
 	/** Remove animator menu action */
 	void RemoveAnimatorAction(FOperatorStackEditorItemPtr InItem) const;
 
+	bool CanExportAnimator(FOperatorStackEditorItemPtr InItem) const;
+	void ExportAnimatorAction(FOperatorStackEditorItemPtr InItem);
+
 	/** Fill item action menus */
 	void FillAddAnimatorMenuSection(UToolMenu* InToolMenu) const;
 	void FillAnimatorHeaderActionMenu(UToolMenu* InToolMenu);
diff --git a/Engine/Plugins/Experimental/PropertyAnimatorCore/Source/PropertyAnimatorCoreEditor/Private/Menus/PropertyAnimatorCoreEditorMenu.cpp b/Engine/Plugins/Experimental/PropertyAnimatorCore/Source/PropertyAnimatorCoreEditor/Private/Menus/PropertyAnimatorCoreEditorMenu.cpp
--- a/Engine/Plugins/Experimental/PropertyAnimatorCore/Source/PropertyAnimatorCoreEditor/Private/Menus/PropertyAnimatorCoreEditorMenu.cpp
+++ b/Engine/Plugins/Experimental/PropertyAnimatorCore/Source/PropertyAnimatorCoreEditor/Private/Menus/PropertyAnimatorCoreEditorMenu.cpp
@@ -7,7 +7,9 @@
 #include "Engine/World.h"
 #include "Misc/Optional.h"
 #include "PropertyHandle.h"
+#include "Presets/PropertyAnimatorCoreAnimatorPreset.h"
 #include "Presets/PropertyAnimatorCorePresetBase.h"
+#include "Presets/PropertyAnimatorCorePropertyPreset.h"
 #include "Styling/SlateIconFinder.h"
 #include "Subsystems/PropertyAnimatorCoreEditorSubsystem.h"
 #include "Subsystems/PropertyAnimatorCoreSubsystem.h"
@@ -40,11 +42,32 @@ void UE::PropertyAnimatorCoreEditor::Menu::FillNewAnimatorSection(UToolMenu* InM
 
 	FToolMenuSection& NewAnimatorsSection = InMenu->FindOrAddSection(TEXT("NewAnimators"), LOCTEXT("NewAnimators.Label", "New Animators"));
 
+	TMap<UPropertyAnimatorCoreBase*, TArray<UPropertyAnimatorCorePresetBase*>> AvailablePresetAnimators;
+
+	UPropertyAnimatorCorePresetBase* EmptyPreset = nullptr;
 	for (UPropertyAnimatorCoreBase* NewAnimator : NewAvailableAnimators)
 	{
-		const FName MenuName = NewAnimator->GetAnimatorOriginalName();
-		const FName MenuCategory = NewAnimator->GetAnimatorCategory();
-		const FText MenuLabel = FText::FromName(MenuName);
+		AvailablePresetAnimators.FindOrAdd(NewAnimator).Add(EmptyPreset);
+	}
+
+	for (UPropertyAnimatorCorePresetBase* Preset : Subsystem->GetAvailablePresets(UPropertyAnimatorCoreAnimatorPreset::StaticClass()))
+	{
+		if (UPropertyAnimatorCoreAnimatorPreset* AnimatorPreset = Cast<UPropertyAnimatorCoreAnimatorPreset>(Preset))
+		{
+			if (UPropertyAnimatorCoreBase* AnimatorTemplate = AnimatorPreset->GetAnimatorTemplate())
+			{
+				AvailablePresetAnimators.FindOrAdd(AnimatorTemplate).Add(AnimatorPreset);
+			}
+		}
+	}
+
+	for (const TPair<UPropertyAnimatorCoreBase*, TArray<UPropertyAnimatorCorePresetBase*>>& NewPresetAnimator : AvailablePresetAnimators)
+	{
+		UPropertyAnimatorCoreBase* NewAnimator = NewPresetAnimator.Key;
+
+		const FName MenuName = NAME_None;
+		FName MenuCategory = NewAnimator->GetAnimatorCategory();
+		FText MenuLabel = FText::FromName(NewAnimator->GetAnimatorOriginalName());
 		const FText MenuTooltip = LOCTEXT("NewAnimator.Tooltip", "Create a new animator");
 		const FSlateIcon MenuIcon = FSlateIconFinder::FindIconForClass(NewAnimator->GetClass());
 
@@ -62,17 +85,24 @@ void UE::PropertyAnimatorCoreEditor::Menu::FillNewAnimatorSection(UToolMenu* InM
 		}
 		else
 		{
-			UPropertyAnimatorCorePresetBase* EmptyPreset = nullptr;
-
-			FToolMenuSection& AnimatorCategorySection = InMenu->FindOrAddSection(MenuCategory, FText::FromName(MenuCategory));
-
-			AnimatorCategorySection.AddMenuEntry(
-				MenuName
-				, MenuLabel
-				, MenuTooltip
-				, MenuIcon
-				, FExecuteAction::CreateLambda(&ExecuteNewAnimatorPresetAction, NewAnimator, ContextActors, EmptyPreset, InMenuData)
-			);
+			for (UPropertyAnimatorCorePresetBase* Preset : NewPresetAnimator.Value)
+			{
+				if (Preset)
+				{
+					MenuLabel = FText::FromString(Preset->GetPresetDisplayName() + TEXT(" (") + NewAnimator->GetAnimatorOriginalName().ToString() + TEXT(")"));
+					MenuCategory = TEXT("Presets");
+				}
+
+				FToolMenuSection& AnimatorCategorySection = InMenu->FindOrAddSection(MenuCategory, FText::FromName(MenuCategory));
+
+				AnimatorCategorySection.AddMenuEntry(
+					MenuName
+					, MenuLabel
+					, MenuTooltip
+					, MenuIcon
+					, FExecuteAction::CreateLambda(&ExecuteNewAnimatorPresetAction, NewAnimator, ContextActors, Preset, InMenuData)
+				);
+			}
 		}
 	}
 }
@@ -396,7 +426,7 @@ void UE::PropertyAnimatorCoreEditor::Menu::ExecuteLinkLastCreatedAnimatorPropert
 	}
 }
 
-void UE::PropertyAnimatorCoreEditor::Menu::ExecuteApplyLastCreatedAnimatorPresetAction(const UPropertyAnimatorCoreBase* InAnimator, UPropertyAnimatorCorePresetBase* InPreset, TSharedRef<FPropertyAnimatorCoreEditorMenuData> InMenuData)
+void UE::PropertyAnimatorCoreEditor::Menu::ExecuteApplyLastCreatedAnimatorPresetAction(const UPropertyAnimatorCoreBase* InAnimator, UPropertyAnimatorCorePropertyPreset* InPreset, TSharedRef<FPropertyAnimatorCoreEditorMenuData> InMenuData)
 {
 	UPropertyAnimatorCoreSubsystem* Subsystem = UPropertyAnimatorCoreSubsystem::Get();
 
@@ -425,7 +455,7 @@ void UE::PropertyAnimatorCoreEditor::Menu::ExecuteApplyLastCreatedAnimatorPreset
 	}
 }
 
-void UE::PropertyAnimatorCoreEditor::Menu::ExecuteLinkAnimatorPresetAction(UPropertyAnimatorCoreBase* InAnimator, UPropertyAnimatorCorePresetBase* InPreset, TSharedRef<FPropertyAnimatorCoreEditorMenuData> InMenuData)
+void UE::PropertyAnimatorCoreEditor::Menu::ExecuteLinkAnimatorPresetAction(UPropertyAnimatorCoreBase* InAnimator, UPropertyAnimatorCorePropertyPreset* InPreset, TSharedRef<FPropertyAnimatorCoreEditorMenuData> InMenuData)
 {
 	UPropertyAnimatorCoreSubsystem* Subsystem = UPropertyAnimatorCoreSubsystem::Get();
 
@@ -448,7 +478,7 @@ void UE::PropertyAnimatorCoreEditor::Menu::ExecuteLinkAnimatorPresetAction(UProp
 	}
 }
 
-void UE::PropertyAnimatorCoreEditor::Menu::ExecuteLinkAnimatorPropertyAction(UPropertyAnimatorCoreBase* InAnimator, FPropertyAnimatorCoreData InProperty, TSharedRef<FPropertyAnimatorCoreEditorMenuData> InMenuData)
+void UE::PropertyAnimatorCoreEditor::Menu::ExecuteLinkAnimatorPropertyAction(UPropertyAnimatorCoreBase* InAnimator, FPropertyAnimatorCoreData InProperty, UPropertyAnimatorCorePropertyPreset* InPreset, TSharedRef<FPropertyAnimatorCoreEditorMenuData> InMenuData)
 {
 	UPropertyAnimatorCoreSubsystem* Subsystem = UPropertyAnimatorCoreSubsystem::Get();
 
@@ -465,10 +495,20 @@ void UE::PropertyAnimatorCoreEditor::Menu::ExecuteLinkAnimatorPropertyAction(UPr
 	if (InAnimator->IsPropertyLinked(InProperty))
 	{
 		Subsystem->UnlinkAnimatorProperty(InAnimator, InProperty, InMenuData->GetOptions().ShouldTransact());
+
+		if (InPreset)
+		{
+			InPreset->OnPresetUnapplied(InAnimator, {InProperty});
+		}
 	}
 	else
 	{
 		Subsystem->LinkAnimatorProperty(InAnimator, InProperty, InMenuData->GetOptions().ShouldTransact());
+
+		if (InPreset)
+		{
+			InPreset->OnPresetApplied(InAnimator, {InProperty});
+		}
 	}
 }
 
@@ -551,7 +591,7 @@ void UE::PropertyAnimatorCoreEditor::Menu::ExecuteDeleteAnimatorAction(UProperty
 	Subsystem->RemoveAnimator(InAnimator, InMenuData->GetOptions().ShouldTransact());
 }
 
-ECheckBoxState UE::PropertyAnimatorCoreEditor::Menu::GetAnimatorPresetState(const UPropertyAnimatorCoreBase* InAnimator, UPropertyAnimatorCorePresetBase* InPreset)
+ECheckBoxState UE::PropertyAnimatorCoreEditor::Menu::GetAnimatorPresetState(const UPropertyAnimatorCoreBase* InAnimator, UPropertyAnimatorCorePropertyPreset* InPreset)
 {
 	if (!IsValid(InAnimator) || InAnimator->IsTemplate() || !InPreset)
 	{
@@ -570,7 +610,7 @@ ECheckBoxState UE::PropertyAnimatorCoreEditor::Menu::GetAnimatorPresetState(cons
 	return !AppliedProperties.IsEmpty() ? ECheckBoxState::Undetermined : ECheckBoxState::Unchecked;
 }
 
-ECheckBoxState UE::PropertyAnimatorCoreEditor::Menu::GetLastAnimatorCreatedPresetState(const UPropertyAnimatorCoreBase* InAnimator, UPropertyAnimatorCorePresetBase* InPreset, TSharedRef<FPropertyAnimatorCoreEditorMenuData> InMenuData)
+ECheckBoxState UE::PropertyAnimatorCoreEditor::Menu::GetLastAnimatorCreatedPresetState(const UPropertyAnimatorCoreBase* InAnimator, UPropertyAnimatorCorePropertyPreset* InPreset, TSharedRef<FPropertyAnimatorCoreEditorMenuData> InMenuData)
 {
 	if (!IsValid(InAnimator)
 		|| !InAnimator->IsTemplate()
@@ -601,7 +641,7 @@ ECheckBoxState UE::PropertyAnimatorCoreEditor::Menu::GetLastAnimatorCreatedPrese
 	return AnimatorsState.Get(ECheckBoxState::Unchecked);
 }
 
-bool UE::PropertyAnimatorCoreEditor::Menu::IsAnimatorPresetLinked(UPropertyAnimatorCoreBase* InAnimator, UPropertyAnimatorCorePresetBase* InPreset)
+bool UE::PropertyAnimatorCoreEditor::Menu::IsAnimatorPresetLinked(UPropertyAnimatorCoreBase* InAnimator, UPropertyAnimatorCorePropertyPreset* InPreset)
 {
 	if (!IsValid(InAnimator) || InAnimator->IsTemplate() || !InPreset)
 	{
@@ -629,7 +669,7 @@ void UE::PropertyAnimatorCoreEditor::Menu::FillNewAnimatorSubmenu(UToolMenu* InM
 
 		FToolMenuSection& PresetSection = InMenu->FindOrAddSection(TEXT("Presets"), LOCTEXT("NewAnimatorPresetsSection.Label", "Presets"));
 
-		UPropertyAnimatorCorePresetBase* EmptyPreset = nullptr;
+		UPropertyAnimatorCorePropertyPreset* EmptyPreset = nullptr;
 
 		PresetSection.AddMenuEntry(
 			TEXT("EmptyPreset")
@@ -644,11 +684,11 @@ void UE::PropertyAnimatorCoreEditor::Menu::FillNewAnimatorSubmenu(UToolMenu* InM
 			)
 		);
 
-		TSet<UPropertyAnimatorCorePresetBase*> SupportedPresets = Subsystem->GetAvailablePresets();
+		TSet<UPropertyAnimatorCorePresetBase*> SupportedPresets = Subsystem->GetAvailablePresets(UPropertyAnimatorCorePropertyPreset::StaticClass());
 
 		for (const AActor* Actor : ContextActors)
 		{
-			SupportedPresets = SupportedPresets.Intersect(Subsystem->GetSupportedPresets(Actor, InAnimator));
+			SupportedPresets = SupportedPresets.Intersect(Subsystem->GetSupportedPresets(Actor, InAnimator, UPropertyAnimatorCorePropertyPreset::StaticClass()));
 		}
 
 		constexpr bool bCloseMenuAfterSelection = false;
@@ -668,7 +708,7 @@ void UE::PropertyAnimatorCoreEditor::Menu::FillNewAnimatorSubmenu(UToolMenu* InM
 				FName(TEXT("Create") + MenuName)
 				, MenuLabel
 				, LOCTEXT("NewAnimatorPresetSection.Tooltip", "Create this animator using this preset")
-				, FNewToolMenuDelegate::CreateLambda(&FillNewPresetAnimatorSubmenu, InAnimator, SupportedPreset, InMenuData)
+				, FNewToolMenuDelegate::CreateLambda(&FillNewPresetAnimatorSubmenu, InAnimator, Cast<UPropertyAnimatorCorePropertyPreset>(SupportedPreset), InMenuData)
 				, bOpenOnClick
 				, FSlateIcon()
 				, bCloseMenuAfterSelection
@@ -741,11 +781,11 @@ void UE::PropertyAnimatorCoreEditor::Menu::FillLinkAnimatorSubmenu(UToolMenu* In
 
 	FToolMenuSection& PresetSection = InMenu->FindOrAddSection(TEXT("Presets"), LOCTEXT("LinkAnimatorPresetsSection.Label", "Presets"));
 
-	TSet<UPropertyAnimatorCorePresetBase*> SupportedPresets = Subsystem->GetAvailablePresets();
+	TSet<UPropertyAnimatorCorePresetBase*> SupportedPresets = Subsystem->GetAvailablePresets(UPropertyAnimatorCorePropertyPreset::StaticClass());
 
 	for (const AActor* Actor : InMenuData->GetContext().GetActors())
 	{
-		SupportedPresets = SupportedPresets.Intersect(Subsystem->GetSupportedPresets(Actor, InAnimator));
+		SupportedPresets = SupportedPresets.Intersect(Subsystem->GetSupportedPresets(Actor, InAnimator, UPropertyAnimatorCorePropertyPreset::StaticClass()));
 	}
 
 	constexpr bool bCloseMenuAfterSelection = false;
@@ -766,7 +806,7 @@ void UE::PropertyAnimatorCoreEditor::Menu::FillLinkAnimatorSubmenu(UToolMenu* In
 			MenuName
 			, MenuLabel
 			, MenuTooltip
-			, FNewToolMenuDelegate::CreateLambda(&FillPresetAnimatorSubmenu, InAnimator, SupportedPreset, InMenuData)
+			, FNewToolMenuDelegate::CreateLambda(&FillPresetAnimatorSubmenu, InAnimator, Cast<UPropertyAnimatorCorePropertyPreset>(SupportedPreset), InMenuData)
 			, bOpenOnClick
 			, FSlateIcon()
 			, bCloseMenuAfterSelection
@@ -781,6 +821,7 @@ void UE::PropertyAnimatorCoreEditor::Menu::FillLinkAnimatorSubmenu(UToolMenu* In
 		InAnimator->GetPropertiesSupported(Property, SupportedProperties, /** SearchDepth */3);
 	}
 
+	UPropertyAnimatorCorePropertyPreset* Empty = nullptr;
 	for (const FPropertyAnimatorCoreData& SupportedProperty : SupportedProperties)
 	{
 		const FString MenuName = SupportedProperty.GetPropertyDisplayName();
@@ -792,7 +833,7 @@ void UE::PropertyAnimatorCoreEditor::Menu::FillLinkAnimatorSubmenu(UToolMenu* In
 			, LOCTEXT("LinkAnimatorPropertySection.Tooltip", "Link or unlink this property from the animator")
 			, FSlateIcon()
 			, FUIAction(
-				FExecuteAction::CreateLambda(&ExecuteLinkAnimatorPropertyAction, InAnimator, SupportedProperty, InMenuData)
+				FExecuteAction::CreateLambda(&ExecuteLinkAnimatorPropertyAction, InAnimator, SupportedProperty, Empty, InMenuData)
 				, FCanExecuteAction()
 				, FIsActionChecked::CreateLambda(&IsAnimatorPropertyLinked, InAnimator, SupportedProperty)
 			)
@@ -801,7 +842,7 @@ void UE::PropertyAnimatorCoreEditor::Menu::FillLinkAnimatorSubmenu(UToolMenu* In
 	}
 }
 
-void UE::PropertyAnimatorCoreEditor::Menu::FillPresetAnimatorSubmenu(UToolMenu* InMenu, UPropertyAnimatorCoreBase* InAnimator, UPropertyAnimatorCorePresetBase* InPreset, TSharedRef<FPropertyAnimatorCoreEditorMenuData> InMenuData)
+void UE::PropertyAnimatorCoreEditor::Menu::FillPresetAnimatorSubmenu(UToolMenu* InMenu, UPropertyAnimatorCoreBase* InAnimator, UPropertyAnimatorCorePropertyPreset* InPreset, TSharedRef<FPropertyAnimatorCoreEditorMenuData> InMenuData)
 {
 	if (!InMenu)
 	{
@@ -845,7 +886,7 @@ void UE::PropertyAnimatorCoreEditor::Menu::FillPresetAnimatorSubmenu(UToolMenu*
 			, MenuTooltip
 			, FSlateIcon()
 			, FUIAction(
-				FExecuteAction::CreateLambda(&ExecuteLinkAnimatorPropertyAction, InAnimator, SupportedProperty, InMenuData)
+				FExecuteAction::CreateLambda(&ExecuteLinkAnimatorPropertyAction, InAnimator, SupportedProperty, InPreset, InMenuData)
 				, FCanExecuteAction()
 				, FIsActionChecked::CreateLambda(&IsAnimatorPropertyLinked, InAnimator, SupportedProperty)
 			)
@@ -856,7 +897,7 @@ void UE::PropertyAnimatorCoreEditor::Menu::FillPresetAnimatorSubmenu(UToolMenu*
 	}
 }
 
-void UE::PropertyAnimatorCoreEditor::Menu::FillNewPresetAnimatorSubmenu(UToolMenu* InMenu, UPropertyAnimatorCoreBase* InAnimator, UPropertyAnimatorCorePresetBase* InPreset, TSharedRef<FPropertyAnimatorCoreEditorMenuData> InMenuData)
+void UE::PropertyAnimatorCoreEditor::Menu::FillNewPresetAnimatorSubmenu(UToolMenu* InMenu, UPropertyAnimatorCoreBase* InAnimator, UPropertyAnimatorCorePropertyPreset* InPreset, TSharedRef<FPropertyAnimatorCoreEditorMenuData> InMenuData)
 {
 	if (!InMenu)
 	{
diff --git a/Engine/Plugins/Experimental/PropertyAnimatorCore/Source/PropertyAnimatorCoreEditor/Private/Menus/PropertyAnimatorCoreEditorMenu.h b/Engine/Plugins/Experimental/PropertyAnimatorCore/Source/PropertyAnimatorCoreEditor/Private/Menus/PropertyAnimatorCoreEditorMenu.h
--- a/Engine/Plugins/Experimental/PropertyAnimatorCore/Source/PropertyAnimatorCoreEditor/Private/Menus/PropertyAnimatorCoreEditorMenu.h
+++ b/Engine/Plugins/Experimental/PropertyAnimatorCore/Source/PropertyAnimatorCoreEditor/Private/Menus/PropertyAnimatorCoreEditorMenu.h
@@ -13,7 +13,7 @@ struct FPropertyAnimatorCoreData;
 class IPropertyHandle;
 class UPropertyAnimatorCoreBase;
 class UPropertyAnimatorCoreEditorMenuContext;
-class UPropertyAnimatorCorePresetBase;
+class UPropertyAnimatorCorePropertyPreset;
 
 namespace UE::PropertyAnimatorCoreEditor::Menu
 {
@@ -37,9 +37,9 @@ namespace UE::PropertyAnimatorCoreEditor::Menu
 
 	void FillLinkAnimatorSubmenu(UToolMenu* InMenu, UPropertyAnimatorCoreBase* InAnimator, TSharedRef<FPropertyAnimatorCoreEditorMenuData> InMenuData);
 
-	void FillPresetAnimatorSubmenu(UToolMenu* InMenu, UPropertyAnimatorCoreBase* InAnimator, UPropertyAnimatorCorePresetBase* InPreset, TSharedRef<FPropertyAnimatorCoreEditorMenuData> InMenuData);
+	void FillPresetAnimatorSubmenu(UToolMenu* InMenu, UPropertyAnimatorCoreBase* InAnimator, UPropertyAnimatorCorePropertyPreset* InPreset, TSharedRef<FPropertyAnimatorCoreEditorMenuData> InMenuData);
 
-	void FillNewPresetAnimatorSubmenu(UToolMenu* InMenu, UPropertyAnimatorCoreBase* InAnimator, UPropertyAnimatorCorePresetBase* InPreset, TSharedRef<FPropertyAnimatorCoreEditorMenuData> InMenuData);
+	void FillNewPresetAnimatorSubmenu(UToolMenu* InMenu, UPropertyAnimatorCoreBase* InAnimator, UPropertyAnimatorCorePropertyPreset* InPreset, TSharedRef<FPropertyAnimatorCoreEditorMenuData> InMenuData);
 
 	/** Execute */
 
@@ -49,11 +49,11 @@ namespace UE::PropertyAnimatorCoreEditor::Menu
 
 	void ExecuteLinkLastCreatedAnimatorPropertyAction(const UPropertyAnimatorCoreBase* InAnimator, FPropertyAnimatorCoreData InProperty, TSharedRef<FPropertyAnimatorCoreEditorMenuData> InMenuData);
 
-	void ExecuteApplyLastCreatedAnimatorPresetAction(const UPropertyAnimatorCoreBase* InAnimator, UPropertyAnimatorCorePresetBase* InPreset, TSharedRef<FPropertyAnimatorCoreEditorMenuData> InMenuData);
+	void ExecuteApplyLastCreatedAnimatorPresetAction(const UPropertyAnimatorCoreBase* InAnimator, UPropertyAnimatorCorePropertyPreset* InPreset, TSharedRef<FPropertyAnimatorCoreEditorMenuData> InMenuData);
 
-	void ExecuteLinkAnimatorPresetAction(UPropertyAnimatorCoreBase* InAnimator, UPropertyAnimatorCorePresetBase* InPreset, TSharedRef<FPropertyAnimatorCoreEditorMenuData> InMenuData);
+	void ExecuteLinkAnimatorPresetAction(UPropertyAnimatorCoreBase* InAnimator, UPropertyAnimatorCorePropertyPreset* InPreset, TSharedRef<FPropertyAnimatorCoreEditorMenuData> InMenuData);
 
-	void ExecuteLinkAnimatorPropertyAction(UPropertyAnimatorCoreBase* InAnimator, FPropertyAnimatorCoreData InProperty, TSharedRef<FPropertyAnimatorCoreEditorMenuData> InMenuData);
+	void ExecuteLinkAnimatorPropertyAction(UPropertyAnimatorCoreBase* InAnimator, FPropertyAnimatorCoreData InProperty, UPropertyAnimatorCorePropertyPreset* InPreset, TSharedRef<FPropertyAnimatorCoreEditorMenuData> InMenuData);
 
 	void ExecuteEnableActorAnimatorAction(TSharedRef<FPropertyAnimatorCoreEditorMenuData> InMenuData, bool bInEnable);
 
@@ -67,17 +67,17 @@ namespace UE::PropertyAnimatorCoreEditor::Menu
 
 	/** Check */
 
-	ECheckBoxState GetAnimatorPresetState(const UPropertyAnimatorCoreBase* InAnimator, UPropertyAnimatorCorePresetBase* InPreset);
+	ECheckBoxState GetAnimatorPresetState(const UPropertyAnimatorCoreBase* InAnimator, UPropertyAnimatorCorePropertyPreset* InPreset);
 
-	ECheckBoxState GetLastAnimatorCreatedPresetState(const UPropertyAnimatorCoreBase* InAnimator, UPropertyAnimatorCorePresetBase* InPreset, TSharedRef<FPropertyAnimatorCoreEditorMenuData> InMenuData);
+	ECheckBoxState GetLastAnimatorCreatedPresetState(const UPropertyAnimatorCoreBase* InAnimator, UPropertyAnimatorCorePropertyPreset* InPreset, TSharedRef<FPropertyAnimatorCoreEditorMenuData> InMenuData);
 
-	bool IsAnimatorPresetLinked(UPropertyAnimatorCoreBase* InAnimator, UPropertyAnimatorCorePresetBase* InPreset);
+	bool IsAnimatorPresetLinked(UPropertyAnimatorCoreBase* InAnimator, UPropertyAnimatorCorePropertyPreset* InPreset);
 
 	bool IsAnimatorPropertyLinked(const UPropertyAnimatorCoreBase* InAnimator, const FPropertyAnimatorCoreData& InProperty);
 
 	bool IsLastAnimatorCreatedPropertyLinked(const UPropertyAnimatorCoreBase* InAnimator, const FPropertyAnimatorCoreData& InProperty, TSharedRef<FPropertyAnimatorCoreEditorMenuData> InMenuData);
 
-	bool IsLastAnimatorCreatedPresetLinked(const UPropertyAnimatorCoreBase* InAnimator, UPropertyAnimatorCorePresetBase* InPreset, TSharedRef<FPropertyAnimatorCoreEditorMenuData> InMenuData);
+	bool IsLastAnimatorCreatedPresetLinked(const UPropertyAnimatorCoreBase* InAnimator, UPropertyAnimatorCorePropertyPreset* InPreset, TSharedRef<FPropertyAnimatorCoreEditorMenuData> InMenuData);
 
 	bool IsLastAnimatorCreatedActionVisible(const UPropertyAnimatorCoreBase* InAnimator, TSharedRef<FPropertyAnimatorCoreEditorMenuData> InMenuData);
 
diff --git a/Engine/Plugins/Experimental/PropertyAnimatorCore/Source/PropertyAnimatorCoreEditor/Private/Presets/AssetDefinition_PropertyAnimatorCorePreset.cpp b/Engine/Plugins/Experimental/PropertyAnimatorCore/Source/PropertyAnimatorCoreEditor/Private/Presets/AssetDefinition_PropertyAnimatorCorePreset.cpp
new file mode 100644
--- /dev/null
+++ b/Engine/Plugins/Experimental/PropertyAnimatorCore/Source/PropertyAnimatorCoreEditor/Private/Presets/AssetDefinition_PropertyAnimatorCorePreset.cpp
@@ -0,0 +1,30 @@
+// Copyright Epic Games, Inc. All Rights Reserved.
+
+#include "AssetDefinition_PropertyAnimatorCorePreset.h"
+
+#include "Presets/PropertyAnimatorCorePresetBase.h"
+
+#define LOCTEXT_NAMESPACE "AssetDefinition_PropertyAnimatorCorePreset"
+
+FText UAssetDefinition_PropertyAnimatorCorePreset::GetAssetDisplayName() const
+{
+	return LOCTEXT("AssetDisplayName", "PropertyAnimatorCorePreset");
+}
+
+TSoftClassPtr<UObject> UAssetDefinition_PropertyAnimatorCorePreset::GetAssetClass() const
+{
+	return UPropertyAnimatorCorePresetBase::StaticClass();
+}
+
+FLinearColor UAssetDefinition_PropertyAnimatorCorePreset::GetAssetColor() const
+{
+	return FLinearColor::White;
+}
+
+TConstArrayView<FAssetCategoryPath> UAssetDefinition_PropertyAnimatorCorePreset::GetAssetCategories() const
+{
+	static const TArray<FAssetCategoryPath> Categories = {EAssetCategoryPaths::Animation};
+	return Categories;
+}
+
+#undef LOCTEXT_NAMESPACE
\ No newline at end of file
diff --git a/Engine/Plugins/Experimental/PropertyAnimatorCore/Source/PropertyAnimatorCoreEditor/Private/Presets/AssetDefinition_PropertyAnimatorCorePreset.h b/Engine/Plugins/Experimental/PropertyAnimatorCore/Source/PropertyAnimatorCoreEditor/Private/Presets/AssetDefinition_PropertyAnimatorCorePreset.h
new file mode 100644
--- /dev/null
+++ b/Engine/Plugins/Experimental/PropertyAnimatorCore/Source/PropertyAnimatorCoreEditor/Private/Presets/AssetDefinition_PropertyAnimatorCorePreset.h
@@ -0,0 +1,20 @@
+// Copyright Epic Games, Inc. All Rights Reserved.
+
+#pragma once
+
+#include "AssetDefinitionDefault.h"
+#include "AssetDefinition_PropertyAnimatorCorePreset.generated.h"
+
+UCLASS()
+class UAssetDefinition_PropertyAnimatorCorePreset : public UAssetDefinitionDefault
+{
+	GENERATED_BODY()
+
+public:
+	//~ Begin UAssetDefinition
+	virtual FText GetAssetDisplayName() const override;
+	virtual TSoftClassPtr<UObject> GetAssetClass() const override;
+	virtual FLinearColor GetAssetColor() const override;
+	virtual TConstArrayView<FAssetCategoryPath> GetAssetCategories() const override;
+	//~ End UAssetDefinition
+};
diff --git a/Engine/Plugins/Experimental/PropertyAnimatorCore/Source/PropertyAnimatorCoreEditor/Private/Presets/PropertyAnimatorCorePresetFactory.cpp b/Engine/Plugins/Experimental/PropertyAnimatorCore/Source/PropertyAnimatorCoreEditor/Private/Presets/PropertyAnimatorCorePresetFactory.cpp
new file mode 100644
--- /dev/null
+++ b/Engine/Plugins/Experimental/PropertyAnimatorCore/Source/PropertyAnimatorCoreEditor/Private/Presets/PropertyAnimatorCorePresetFactory.cpp
@@ -0,0 +1,99 @@
+// Copyright Epic Games, Inc. All Rights Reserved.
+
+#include "PropertyAnimatorCorePresetFactory.h"
+
+#include "ClassViewerFilter.h"
+#include "ClassViewerModule.h"
+#include "Kismet2/KismetEditorUtilities.h"
+#include "Kismet2/SClassPickerDialog.h"
+#include "Modules/ModuleManager.h"
+#include "Presets/PropertyAnimatorCorePresetBase.h"
+
+#define LOCTEXT_NAMESPACE "PropertyAnimatorCorePresetFactory"
+
+UPropertyAnimatorCorePresetFactory::UPropertyAnimatorCorePresetFactory()
+{
+	SupportedClass = UPropertyAnimatorCorePresetBase::StaticClass();
+	bCreateNew = true;
+	bEditAfterNew = false;
+	bEditorImport = false;
+	bText = true;
+}
+
+bool UPropertyAnimatorCorePresetFactory::ConfigureProperties()
+{
+	NewPresetClass = nullptr;
+
+	// Load the class viewer module to display a class picker
+	FModuleManager::LoadModuleChecked<FClassViewerModule>("ClassViewer");
+
+	class FAssetClassParentFilter : public IClassViewerFilter
+	{
+	public:
+		FAssetClassParentFilter()
+			: DisallowedClassFlags(CLASS_None)
+		{}
+
+		virtual bool IsClassAllowed(const FClassViewerInitializationOptions& InInitOptions, const UClass* InClass, TSharedRef<FClassViewerFilterFuncs> InFilterFunc) override
+		{
+			return !InClass->HasAnyClassFlags(DisallowedClassFlags)
+				&& InClass->CanCreateAssetOfClass()
+				&& InFilterFunc->IfInChildOfClassesSet(AllowedChildrenOfClasses, InClass) != EFilterReturn::Failed
+				&& !FKismetEditorUtilities::CanCreateBlueprintOfClass(InClass);
+		}
+
+		virtual bool IsUnloadedClassAllowed(const FClassViewerInitializationOptions& InInitOptions, const TSharedRef<const IUnloadedBlueprintData> InUnloadedClassData, TSharedRef<FClassViewerFilterFuncs> InFilterFunc) override
+		{
+			// Disallow blueprint base class
+			return false;
+		}
+
+		TSet<const UClass*> AllowedChildrenOfClasses;
+		EClassFlags DisallowedClassFlags;
+	};
+
+	const TSharedPtr<FAssetClassParentFilter> ClassFilter = MakeShareable(new FAssetClassParentFilter());
+	ClassFilter->DisallowedClassFlags = CLASS_Abstract | CLASS_Deprecated | CLASS_NewerVersionExists | CLASS_Transient;
+	ClassFilter->AllowedChildrenOfClasses.Add(UPropertyAnimatorCorePresetBase::StaticClass());
+
+	FClassViewerInitializationOptions Options;
+	Options.Mode = EClassViewerMode::ClassPicker;
+	Options.ClassFilters.Add(ClassFilter.ToSharedRef());
+
+	const FText TitleText = LOCTEXT("CreatePropertyAnimatorCorePreset", "Pick Preset Class");
+
+	UClass* PickedClass = nullptr;
+	const bool bSuccess = SClassPickerDialog::PickClass(TitleText, Options, PickedClass, UPropertyAnimatorCorePresetBase::StaticClass());
+
+	if (bSuccess)
+	{
+		NewPresetClass = PickedClass;
+	}
+
+	return bSuccess;
+}
+
+UObject* UPropertyAnimatorCorePresetFactory::FactoryCreateNew(UClass* InClass, UObject* InParent, FName InName, EObjectFlags InFlags, UObject* InContext, FFeedbackContext* InWarn)
+{
+	UPropertyAnimatorCorePresetBase* NewObjectAsset = nullptr;
+
+	if (NewPresetClass)
+	{
+		NewObjectAsset = NewObject<UPropertyAnimatorCorePresetBase>(InParent, NewPresetClass, InName, InFlags | RF_Transactional);
+	}
+
+	return NewObjectAsset;
+}
+
+UObject* UPropertyAnimatorCorePresetFactory::FactoryCreateText(UClass* InClass, UObject* InParent, FName InName, EObjectFlags InFlags,
+	UObject* InContext, const TCHAR* InType, const TCHAR*& InBuffer, const TCHAR* InBufferEnd, FFeedbackContext* InWarn)
+{
+	return nullptr;
+}
+
+bool UPropertyAnimatorCorePresetFactory::FactoryCanImport(const FString& InFilename)
+{
+	return false;
+}
+
+#undef LOCTEXT_NAMESPACE
diff --git a/Engine/Plugins/Experimental/PropertyAnimatorCore/Source/PropertyAnimatorCoreEditor/Private/Presets/PropertyAnimatorCorePresetFactory.h b/Engine/Plugins/Experimental/PropertyAnimatorCore/Source/PropertyAnimatorCoreEditor/Private/Presets/PropertyAnimatorCorePresetFactory.h
new file mode 100644
--- /dev/null
+++ b/Engine/Plugins/Experimental/PropertyAnimatorCore/Source/PropertyAnimatorCoreEditor/Private/Presets/PropertyAnimatorCorePresetFactory.h
@@ -0,0 +1,26 @@
+// Copyright Epic Games, Inc. All Rights Reserved.
+
+#pragma once
+
+#include "Factories/Factory.h"
+#include "PropertyAnimatorCorePresetFactory.generated.h"
+
+class UPropertyAnimatorCorePresetBase;
+
+UCLASS()
+class UPropertyAnimatorCorePresetFactory : public UFactory
+{
+	GENERATED_BODY()
+
+public:
+	UPropertyAnimatorCorePresetFactory();
+
+	//~ Begin UFactory
+	virtual bool ConfigureProperties() override;
+	virtual UObject* FactoryCreateNew(UClass* InClass, UObject* InParent, FName InName, EObjectFlags InFlags, UObject* InContext, FFeedbackContext* InWarn) override;
+	virtual UObject* FactoryCreateText(UClass* InClass, UObject* InParent, FName InName, EObjectFlags InFlags, UObject* InContext, const TCHAR* InType, const TCHAR*& InBuffer, const TCHAR* InBufferEnd, FFeedbackContext* InWarn) override;
+	virtual bool FactoryCanImport(const FString& InFilename) override;
+	//~ End UFactory
+
+	TSubclassOf<UPropertyAnimatorCorePresetBase> NewPresetClass;
+};
diff --git a/Engine/Plugins/Experimental/PropertyAnimatorCore/Source/PropertyAnimatorCoreEditor/Private/Styles/PropertyAnimatorCoreEditorStyle.cpp b/Engine/Plugins/Experimental/PropertyAnimatorCore/Source/PropertyAnimatorCoreEditor/Private/Styles/PropertyAnimatorCoreEditorStyle.cpp
--- a/Engine/Plugins/Experimental/PropertyAnimatorCore/Source/PropertyAnimatorCoreEditor/Private/Styles/PropertyAnimatorCoreEditorStyle.cpp
+++ b/Engine/Plugins/Experimental/PropertyAnimatorCore/Source/PropertyAnimatorCoreEditor/Private/Styles/PropertyAnimatorCoreEditorStyle.cpp
@@ -25,6 +25,7 @@ FPropertyAnimatorCoreEditorStyle::FPropertyAnimatorCoreEditorStyle()
 	Set("ClassIcon.PropertyAnimatorCoreBase",       new IMAGE_BRUSH_SVG("EditorIcons/PropertyAnimatorCoreDefault", Icon16x16));
 	Set("PropertyControlIcon.Default",              new IMAGE_BRUSH_SVG("EditorIcons/PropertyAnimatorCoreDefault", Icon16x16));
 	Set("PropertyControlIcon.Linked",               new IMAGE_BRUSH_SVG("EditorIcons/PropertyAnimatorCoreLinked",  Icon16x16));
+	Set("PropertyControlIcon.Export",               new IMAGE_BRUSH_SVG("EditorIcons/Export",                      Icon16x16));
 	Set("ManualTimeSourceControl.PlayForward",      new IMAGE_BRUSH_SVG("EditorIcons/Play",                        Icon16x16));
 	Set("ManualTimeSourceControl.PlayBackward",     new IMAGE_BRUSH_SVG("EditorIcons/PlayReverse",                 Icon16x16));
 	Set("ManualTimeSourceControl.Pause",            new IMAGE_BRUSH_SVG("EditorIcons/Pause",                       Icon16x16));
diff --git a/Engine/Plugins/Experimental/PropertyAnimatorCore/Source/PropertyAnimatorCoreEditor/Private/Subsystems/PropertyAnimatorCoreEditorSubsystem.cpp b/Engine/Plugins/Experimental/PropertyAnimatorCore/Source/PropertyAnimatorCoreEditor/Private/Subsystems/PropertyAnimatorCoreEditorSubsystem.cpp
--- a/Engine/Plugins/Experimental/PropertyAnimatorCore/Source/PropertyAnimatorCoreEditor/Private/Subsystems/PropertyAnimatorCoreEditorSubsystem.cpp
+++ b/Engine/Plugins/Experimental/PropertyAnimatorCore/Source/PropertyAnimatorCoreEditor/Private/Subsystems/PropertyAnimatorCoreEditorSubsystem.cpp
@@ -2,8 +2,10 @@
 
 #include "Subsystems/PropertyAnimatorCoreEditorSubsystem.h"
 
+#include "AssetRegistry/AssetRegistryModule.h"
 #include "Components/PropertyAnimatorCoreComponent.h"
 #include "DetailRowMenuContext.h"
+#include "Dialogs/DlgPickAssetPath.h"
 #include "Editor.h"
 #include "Framework/Application/SlateApplication.h"
 #include "IDetailTreeNode.h"
@@ -155,6 +157,57 @@ bool UPropertyAnimatorCoreEditorSubsystem::FillAnimatorMenu(UToolMenu* InMenu, c
 	return false;
 }
 
+UPropertyAnimatorCorePresetBase* UPropertyAnimatorCoreEditorSubsystem::CreatePresetAsset(TSubclassOf<UPropertyAnimatorCorePresetBase> InPresetClass, const TArray<IPropertyAnimatorCorePresetable*>& InPresetables)
+{
+	UPropertyAnimatorCorePresetBase* NewPreset = nullptr;
+
+	if (!InPresetClass.Get() || InPresetables.IsEmpty())
+	{
+		return NewPreset;
+	}
+
+	// Pick asset path and name
+	FString PickedPath;
+	FString PickedName;
+	{
+		TSharedPtr<SDlgPickAssetPath> DialogWidget = SNew(SDlgPickAssetPath)
+		.Title(LOCTEXT("PickAssetsLocation", "Choose preset name and location"))
+		.DefaultAssetPath(FText::FromString(TEXT("/PropertyAnimatorCore/Presets/NewPreset")))
+		.AllowReadOnlyFolders(true);
+
+		if (DialogWidget->ShowModal() != EAppReturnType::Ok)
+		{
+			PickedPath = TEXT("");
+			return NewPreset;
+		}
+
+		PickedPath = DialogWidget->GetAssetPath().ToString();
+		PickedName = DialogWidget->GetAssetName().ToString();
+	}
+
+	if (PickedPath.IsEmpty() || PickedName.IsEmpty())
+	{
+		return NewPreset;
+	}
+
+	// Find/create package
+	UPackage* Package = CreatePackage(*(PickedPath + TEXT("/") + PickedName));
+
+	if (!Package)
+	{
+		return nullptr;
+	}
+
+	NewPreset = NewObject<UPropertyAnimatorCorePresetBase>(Package, InPresetClass.Get(), FName(PickedName), RF_Public | RF_Standalone);
+	NewPreset->CreatePreset(FName(PickedName), InPresetables);
+	NewPreset->MarkPackageDirty();
+
+	// Notify asset registry of new asset
+	FAssetRegistryModule::AssetCreated(NewPreset);
+
+	return NewPreset;
+}
+
 UPropertyAnimatorCoreEditorSubsystem* UPropertyAnimatorCoreEditorSubsystem::Get()
 {
 	if (GEditor)
diff --git a/Engine/Plugins/Experimental/PropertyAnimatorCore/Source/PropertyAnimatorCoreEditor/PropertyAnimatorCoreEditor.Build.cs b/Engine/Plugins/Experimental/PropertyAnimatorCore/Source/PropertyAnimatorCoreEditor/PropertyAnimatorCoreEditor.Build.cs
--- a/Engine/Plugins/Experimental/PropertyAnimatorCore/Source/PropertyAnimatorCoreEditor/PropertyAnimatorCoreEditor.Build.cs
+++ b/Engine/Plugins/Experimental/PropertyAnimatorCore/Source/PropertyAnimatorCoreEditor/PropertyAnimatorCoreEditor.Build.cs
@@ -19,6 +19,7 @@ public class PropertyAnimatorCoreEditor : ModuleRules
 	        new string[]
 	        {
 		        "ApplicationCore",
+				"AssetDefinition",
 		        "CoreUObject",
 		        "EditorSubsystem",
 		        "EditorWidgets",
diff --git a/Engine/Plugins/Experimental/PropertyAnimatorCore/Source/PropertyAnimatorCoreEditor/Public/Subsystems/PropertyAnimatorCoreEditorSubsystem.h b/Engine/Plugins/Experimental/PropertyAnimatorCore/Source/PropertyAnimatorCoreEditor/Public/Subsystems/PropertyAnimatorCoreEditorSubsystem.h
--- a/Engine/Plugins/Experimental/PropertyAnimatorCore/Source/PropertyAnimatorCoreEditor/Public/Subsystems/PropertyAnimatorCoreEditorSubsystem.h
+++ b/Engine/Plugins/Experimental/PropertyAnimatorCore/Source/PropertyAnimatorCoreEditor/Public/Subsystems/PropertyAnimatorCoreEditorSubsystem.h
@@ -9,11 +9,13 @@
 #include "PropertyAnimatorCoreEditorSubsystem.generated.h"
 
 class IDetailTreeNode;
+class IPropertyAnimatorCorePresetable;
 class IPropertyHandle;
 class SPropertyAnimatorCoreEditorEditPanel;
 class SWidget;
 class UPropertyAnimatorCoreBase;
 class UPropertyAnimatorCoreEditorMenuContext;
+class UPropertyAnimatorCorePresetBase;
 class UToolMenu;
 struct FPropertyAnimatorCoreData;
 struct FPropertyAnimatorCoreEditorEditPanelOptions;
@@ -38,6 +40,9 @@ public:
 	/** Fills a menu based on context objects and menu options */
 	PROPERTYANIMATORCOREEDITOR_API bool FillAnimatorMenu(UToolMenu* InMenu, const FPropertyAnimatorCoreEditorMenuContext& InContext, const FPropertyAnimatorCoreEditorMenuOptions& InOptions);
 
+	/** Creates a preset asset for an item */
+	PROPERTYANIMATORCOREEDITOR_API UPropertyAnimatorCorePresetBase* CreatePresetAsset(TSubclassOf<UPropertyAnimatorCorePresetBase> InPresetClass, const TArray<IPropertyAnimatorCorePresetable*>& InPresetables);
+
 protected:
 	/** Setup details panel button customization */
 	void RegisterDetailPanelCustomization();
