MotionDesign : PropertyAnimatorCore
- Reworked preset system with custom archive implementation that can be swapped

diff --git a/Engine/Plugins/Experimental/PropertyAnimator/Source/PropertyAnimator/Private/Animators/PropertyAnimatorClock.cpp b/Engine/Plugins/Experimental/PropertyAnimator/Source/PropertyAnimator/Private/Animators/PropertyAnimatorClock.cpp
--- a/Engine/Plugins/Experimental/PropertyAnimator/Source/PropertyAnimator/Private/Animators/PropertyAnimatorClock.cpp
+++ b/Engine/Plugins/Experimental/PropertyAnimator/Source/PropertyAnimator/Private/Animators/PropertyAnimatorClock.cpp
@@ -2,8 +2,6 @@
 
 #include "Animators/PropertyAnimatorClock.h"
 
-#include "Dom/JsonObject.h"
-#include "Dom/JsonValue.h"
 #include "Misc/DateTime.h"
 #include "Properties/Converters/PropertyAnimatorCoreConverterBase.h"
 #include "Properties/Handlers/PropertyAnimatorCoreHandlerBase.h"
@@ -118,15 +116,15 @@ void UPropertyAnimatorClock::EvaluateProperties(FInstancedPropertyBag& InParamet
 	});
 }
 
-bool UPropertyAnimatorClock::ImportPreset(const UPropertyAnimatorCorePresetBase* InPreset, const TSharedRef<FJsonValue>& InValue)
+bool UPropertyAnimatorClock::ImportPreset(const UPropertyAnimatorCorePresetBase* InPreset, const TSharedRef<FPropertyAnimatorCorePresetArchive>& InValue)
 {
-	const TSharedPtr<FJsonObject>* JsonAnimatorObject;
-
-	if (Super::ImportPreset(InPreset, InValue) && InValue->TryGetObject(JsonAnimatorObject))
+	if (Super::ImportPreset(InPreset, InValue) && InValue->IsObject())
 	{
-		FString JsonDisplayFormat = DisplayFormat;
-		(*JsonAnimatorObject)->TryGetStringField(GET_MEMBER_NAME_STRING_CHECKED(UPropertyAnimatorClock, DisplayFormat), JsonDisplayFormat);
-		SetDisplayFormat(JsonDisplayFormat);
+		const TSharedPtr<FPropertyAnimatorCorePresetObjectArchive> AnimatorArchive = InValue->AsMutableObject();
+
+		FString DisplayFormatValue = DisplayFormat;
+		AnimatorArchive->Get(GET_MEMBER_NAME_STRING_CHECKED(UPropertyAnimatorClock, DisplayFormat), DisplayFormatValue);
+		SetDisplayFormat(DisplayFormatValue);
 
 		return true;
 	}
@@ -134,13 +132,13 @@ bool UPropertyAnimatorClock::ImportPreset(const UPropertyAnimatorCorePresetBase*
 	return false;
 }
 
-bool UPropertyAnimatorClock::ExportPreset(const UPropertyAnimatorCorePresetBase* InPreset, TSharedPtr<FJsonValue>& OutValue)
+bool UPropertyAnimatorClock::ExportPreset(const UPropertyAnimatorCorePresetBase* InPreset, TSharedPtr<FPropertyAnimatorCorePresetArchive>& OutValue) const
 {
-	const TSharedPtr<FJsonObject>* JsonAnimatorObject;
-
-	if (Super::ExportPreset(InPreset, OutValue) && OutValue->TryGetObject(JsonAnimatorObject))
+	if (Super::ExportPreset(InPreset, OutValue) && OutValue->IsObject())
 	{
-		(*JsonAnimatorObject)->SetStringField(GET_MEMBER_NAME_STRING_CHECKED(UPropertyAnimatorClock, DisplayFormat), DisplayFormat);
+		const TSharedPtr<FPropertyAnimatorCorePresetObjectArchive> AnimatorArchive = OutValue->AsMutableObject();
+
+		AnimatorArchive->Set(GET_MEMBER_NAME_STRING_CHECKED(UPropertyAnimatorClock, DisplayFormat), DisplayFormat);
 
 		return true;
 	}
diff --git a/Engine/Plugins/Experimental/PropertyAnimator/Source/PropertyAnimator/Private/Animators/PropertyAnimatorCounter.cpp b/Engine/Plugins/Experimental/PropertyAnimator/Source/PropertyAnimator/Private/Animators/PropertyAnimatorCounter.cpp
--- a/Engine/Plugins/Experimental/PropertyAnimator/Source/PropertyAnimator/Private/Animators/PropertyAnimatorCounter.cpp
+++ b/Engine/Plugins/Experimental/PropertyAnimator/Source/PropertyAnimator/Private/Animators/PropertyAnimatorCounter.cpp
@@ -2,8 +2,6 @@
 
 #include "Animators/PropertyAnimatorCounter.h"
 
-#include "Dom/JsonObject.h"
-#include "Dom/JsonValue.h"
 #include "Properties/Converters/PropertyAnimatorCoreConverterBase.h"
 #include "Settings/PropertyAnimatorSettings.h"
 #include "Subsystems/PropertyAnimatorCoreSubsystem.h"
@@ -283,39 +281,48 @@ void UPropertyAnimatorCounter::EvaluateProperties(FInstancedPropertyBag& InParam
 	});
 }
 
-bool UPropertyAnimatorCounter::ImportPreset(const UPropertyAnimatorCorePresetBase* InPreset, const TSharedRef<FJsonValue>& InValue)
+bool UPropertyAnimatorCounter::ImportPreset(const UPropertyAnimatorCorePresetBase* InPreset, const TSharedRef<FPropertyAnimatorCorePresetArchive>& InValue)
 {
-	const TSharedPtr<FJsonObject>* JsonAnimatorObject;
-
-	if (Super::ImportPreset(InPreset, InValue) && InValue->TryGetObject(JsonAnimatorObject))
+	if (Super::ImportPreset(InPreset, InValue) && InValue->IsObject())
 	{
-		FString JsonDisplayPattern = DisplayPattern.ToString();
-		(*JsonAnimatorObject)->TryGetStringField(GET_MEMBER_NAME_STRING_CHECKED(UPropertyAnimatorCounter, DisplayPattern), JsonDisplayPattern);
-		SetDisplayPattern(FText::FromString(JsonDisplayPattern));
+		const TSharedPtr<FPropertyAnimatorCorePresetObjectArchive> AnimatorArchive = InValue->AsMutableObject();
+
+		FString DisplayPatternValue = DisplayPattern.ToString();
+		AnimatorArchive->Get(GET_MEMBER_NAME_STRING_CHECKED(UPropertyAnimatorCounter, DisplayPattern), DisplayPatternValue);
+		SetDisplayPattern(FText::FromString(DisplayPatternValue));
 
-		FString JsonPresetFormatNameStr = PresetFormatName.ToString();
-		(*JsonAnimatorObject)->TryGetStringField(GET_MEMBER_NAME_STRING_CHECKED(UPropertyAnimatorCounter, PresetFormatName), JsonPresetFormatNameStr);
-		FName JsonPresetFormatName(JsonPresetFormatNameStr);
-		SetPresetFormatName(JsonPresetFormatName);
+		FString PresetFormatNameValue = PresetFormatName.ToString();
+		AnimatorArchive->Get(GET_MEMBER_NAME_STRING_CHECKED(UPropertyAnimatorCounter, PresetFormatName), PresetFormatNameValue);
+		FName NewPresetFormatName(PresetFormatNameValue);
+		SetPresetFormatName(NewPresetFormatName);
 
-		if (!PresetFormatName.IsEqual(JsonPresetFormatName))
+		if (!PresetFormatName.IsEqual(NewPresetFormatName))
 		{
 			FPropertyAnimatorCounterFormat Format;
-			Format.FormatName = JsonPresetFormatName;
+			Format.FormatName = NewPresetFormatName;
+
+			uint64 IntegerCount = Format.MinIntegerCount;
+			AnimatorArchive->Get(GET_MEMBER_NAME_STRING_CHECKED(FPropertyAnimatorCounterFormat, MinIntegerCount), IntegerCount);
+			Format.MinIntegerCount = IntegerCount;
+
+			uint64 DecimalCount = Format.MaxDecimalCount;
+			AnimatorArchive->Get(GET_MEMBER_NAME_STRING_CHECKED(FPropertyAnimatorCounterFormat, MaxDecimalCount), DecimalCount);
+			Format.MaxDecimalCount = DecimalCount;
 
-			(*JsonAnimatorObject)->TryGetNumberField(GET_MEMBER_NAME_STRING_CHECKED(FPropertyAnimatorCounterFormat, MinIntegerCount), Format.MinIntegerCount);
-			(*JsonAnimatorObject)->TryGetNumberField(GET_MEMBER_NAME_STRING_CHECKED(FPropertyAnimatorCounterFormat, MaxDecimalCount), Format.MaxDecimalCount);
-			(*JsonAnimatorObject)->TryGetNumberField(GET_MEMBER_NAME_STRING_CHECKED(FPropertyAnimatorCounterFormat, GroupingSize), Format.GroupingSize);
-			(*JsonAnimatorObject)->TryGetStringField(GET_MEMBER_NAME_STRING_CHECKED(FPropertyAnimatorCounterFormat, DecimalCharacter), Format.DecimalCharacter);
-			(*JsonAnimatorObject)->TryGetStringField(GET_MEMBER_NAME_STRING_CHECKED(FPropertyAnimatorCounterFormat, GroupingCharacter), Format.GroupingCharacter);
-			(*JsonAnimatorObject)->TryGetStringField(GET_MEMBER_NAME_STRING_CHECKED(FPropertyAnimatorCounterFormat, PaddingCharacter), Format.PaddingCharacter);
+			uint64 GroupingSize = Format.GroupingSize;
+			AnimatorArchive->Get(GET_MEMBER_NAME_STRING_CHECKED(FPropertyAnimatorCounterFormat, GroupingSize), GroupingSize);
+			Format.GroupingSize = GroupingSize;
 
-			uint8 JsonRoundingMode = static_cast<uint8>(Format.RoundingMode);
-			(*JsonAnimatorObject)->TryGetNumberField(GET_MEMBER_NAME_STRING_CHECKED(FPropertyAnimatorCounterFormat, RoundingMode), JsonRoundingMode);
-			Format.RoundingMode = static_cast<EPropertyAnimatorCounterRoundingMode>(JsonRoundingMode);
+			AnimatorArchive->Get(GET_MEMBER_NAME_STRING_CHECKED(FPropertyAnimatorCounterFormat, DecimalCharacter), Format.DecimalCharacter);
+			AnimatorArchive->Get(GET_MEMBER_NAME_STRING_CHECKED(FPropertyAnimatorCounterFormat, GroupingCharacter), Format.GroupingCharacter);
+			AnimatorArchive->Get(GET_MEMBER_NAME_STRING_CHECKED(FPropertyAnimatorCounterFormat, PaddingCharacter), Format.PaddingCharacter);
 
-			(*JsonAnimatorObject)->TryGetBoolField(GET_MEMBER_NAME_STRING_CHECKED(FPropertyAnimatorCounterFormat, bUseSign), Format.bUseSign);
-			(*JsonAnimatorObject)->TryGetBoolField(GET_MEMBER_NAME_STRING_CHECKED(FPropertyAnimatorCounterFormat, bTruncate), Format.bTruncate);
+			uint64 RoundingModeValue = static_cast<uint64>(Format.RoundingMode);
+			AnimatorArchive->Get(GET_MEMBER_NAME_STRING_CHECKED(FPropertyAnimatorCounterFormat, RoundingMode), RoundingModeValue);
+			Format.RoundingMode = static_cast<EPropertyAnimatorCounterRoundingMode>(RoundingModeValue);
+
+			AnimatorArchive->Get(GET_MEMBER_NAME_STRING_CHECKED(FPropertyAnimatorCounterFormat, bUseSign), Format.bUseSign);
+			AnimatorArchive->Get(GET_MEMBER_NAME_STRING_CHECKED(FPropertyAnimatorCounterFormat, bTruncate), Format.bTruncate);
 
 			SetUseCustomFormat(true);
 			SetCustomFormat(&Format);
@@ -327,26 +334,26 @@ bool UPropertyAnimatorCounter::ImportPreset(const UPropertyAnimatorCorePresetBas
 	return false;
 }
 
-bool UPropertyAnimatorCounter::ExportPreset(const UPropertyAnimatorCorePresetBase* InPreset, TSharedPtr<FJsonValue>& OutValue)
+bool UPropertyAnimatorCounter::ExportPreset(const UPropertyAnimatorCorePresetBase* InPreset, TSharedPtr<FPropertyAnimatorCorePresetArchive>& OutValue) const
 {
-	const TSharedPtr<FJsonObject>* JsonAnimatorObject;
-
-	if (Super::ExportPreset(InPreset, OutValue) && OutValue->TryGetObject(JsonAnimatorObject))
+	if (Super::ExportPreset(InPreset, OutValue) && OutValue->IsObject())
 	{
-		(*JsonAnimatorObject)->SetStringField(GET_MEMBER_NAME_STRING_CHECKED(UPropertyAnimatorCounter, DisplayPattern), DisplayPattern.ToString());
+		const TSharedPtr<FPropertyAnimatorCorePresetObjectArchive> AnimatorArchive = OutValue->AsMutableObject();
+
+		AnimatorArchive->Set(GET_MEMBER_NAME_STRING_CHECKED(UPropertyAnimatorCounter, DisplayPattern), DisplayPattern.ToString());
 
 		if (const FPropertyAnimatorCounterFormat* Format = GetFormat())
 		{
-			(*JsonAnimatorObject)->SetStringField(GET_MEMBER_NAME_STRING_CHECKED(UPropertyAnimatorCounter, PresetFormatName), bUseCustomFormat ? Format->FormatName.ToString() : PresetFormatName.ToString());
-			(*JsonAnimatorObject)->SetNumberField(GET_MEMBER_NAME_STRING_CHECKED(FPropertyAnimatorCounterFormat, MinIntegerCount), Format->MinIntegerCount);
-			(*JsonAnimatorObject)->SetNumberField(GET_MEMBER_NAME_STRING_CHECKED(FPropertyAnimatorCounterFormat, MaxDecimalCount), Format->MaxDecimalCount);
-			(*JsonAnimatorObject)->SetNumberField(GET_MEMBER_NAME_STRING_CHECKED(FPropertyAnimatorCounterFormat, GroupingSize), Format->GroupingSize);
-			(*JsonAnimatorObject)->SetStringField(GET_MEMBER_NAME_STRING_CHECKED(FPropertyAnimatorCounterFormat, DecimalCharacter), Format->DecimalCharacter);
-			(*JsonAnimatorObject)->SetStringField(GET_MEMBER_NAME_STRING_CHECKED(FPropertyAnimatorCounterFormat, GroupingCharacter), Format->GroupingCharacter);
-			(*JsonAnimatorObject)->SetStringField(GET_MEMBER_NAME_STRING_CHECKED(FPropertyAnimatorCounterFormat, PaddingCharacter), Format->PaddingCharacter);
-			(*JsonAnimatorObject)->SetNumberField(GET_MEMBER_NAME_STRING_CHECKED(FPropertyAnimatorCounterFormat, RoundingMode), static_cast<uint8>(Format->RoundingMode));
-			(*JsonAnimatorObject)->SetBoolField(GET_MEMBER_NAME_STRING_CHECKED(FPropertyAnimatorCounterFormat, bUseSign), Format->bUseSign);
-			(*JsonAnimatorObject)->SetBoolField(GET_MEMBER_NAME_STRING_CHECKED(FPropertyAnimatorCounterFormat, bTruncate), Format->bTruncate);
+			AnimatorArchive->Set(GET_MEMBER_NAME_STRING_CHECKED(UPropertyAnimatorCounter, PresetFormatName), bUseCustomFormat ? Format->FormatName.ToString() : PresetFormatName.ToString());
+			AnimatorArchive->Set(GET_MEMBER_NAME_STRING_CHECKED(FPropertyAnimatorCounterFormat, MinIntegerCount), static_cast<uint64>(Format->MinIntegerCount));
+			AnimatorArchive->Set(GET_MEMBER_NAME_STRING_CHECKED(FPropertyAnimatorCounterFormat, MaxDecimalCount), static_cast<uint64>(Format->MaxDecimalCount));
+			AnimatorArchive->Set(GET_MEMBER_NAME_STRING_CHECKED(FPropertyAnimatorCounterFormat, GroupingSize), static_cast<uint64>(Format->GroupingSize));
+			AnimatorArchive->Set(GET_MEMBER_NAME_STRING_CHECKED(FPropertyAnimatorCounterFormat, DecimalCharacter), Format->DecimalCharacter);
+			AnimatorArchive->Set(GET_MEMBER_NAME_STRING_CHECKED(FPropertyAnimatorCounterFormat, GroupingCharacter), Format->GroupingCharacter);
+			AnimatorArchive->Set(GET_MEMBER_NAME_STRING_CHECKED(FPropertyAnimatorCounterFormat, PaddingCharacter), Format->PaddingCharacter);
+			AnimatorArchive->Set(GET_MEMBER_NAME_STRING_CHECKED(FPropertyAnimatorCounterFormat, RoundingMode), static_cast<uint64>(Format->RoundingMode));
+			AnimatorArchive->Set(GET_MEMBER_NAME_STRING_CHECKED(FPropertyAnimatorCounterFormat, bUseSign), Format->bUseSign);
+			AnimatorArchive->Set(GET_MEMBER_NAME_STRING_CHECKED(FPropertyAnimatorCounterFormat, bTruncate), Format->bTruncate);
 		}
 
 		return true;
diff --git a/Engine/Plugins/Experimental/PropertyAnimator/Source/PropertyAnimator/Private/Animators/PropertyAnimatorCurve.cpp b/Engine/Plugins/Experimental/PropertyAnimator/Source/PropertyAnimator/Private/Animators/PropertyAnimatorCurve.cpp
--- a/Engine/Plugins/Experimental/PropertyAnimator/Source/PropertyAnimator/Private/Animators/PropertyAnimatorCurve.cpp
+++ b/Engine/Plugins/Experimental/PropertyAnimator/Source/PropertyAnimator/Private/Animators/PropertyAnimatorCurve.cpp
@@ -4,8 +4,6 @@
 
 #include "Curves/PropertyAnimatorEaseCurve.h"
 #include "Curves/PropertyAnimatorWaveCurve.h"
-#include "Dom/JsonObject.h"
-#include "Dom/JsonValue.h"
 #include "Properties/PropertyAnimatorFloatContext.h"
 #include "UObject/ConstructorHelpers.h"
 
@@ -124,64 +122,64 @@ bool UPropertyAnimatorCurve::EvaluateProperty(const FPropertyAnimatorCoreData& I
 	return InContext->EvaluateProperty(InPropertyData, InParameters, OutEvaluationResult);
 }
 
-bool UPropertyAnimatorCurve::ImportPreset(const UPropertyAnimatorCorePresetBase* InPreset, const TSharedRef<FJsonValue>& InValue)
+bool UPropertyAnimatorCurve::ImportPreset(const UPropertyAnimatorCorePresetBase* InPreset, const TSharedRef<FPropertyAnimatorCorePresetArchive>& InValue)
 {
-	const TSharedPtr<FJsonObject>* JsonAnimatorObject;
-
-	if (Super::ImportPreset(InPreset, InValue) && InValue->TryGetObject(JsonAnimatorObject))
+	if (Super::ImportPreset(InPreset, InValue) && InValue->IsObject())
 	{
-		FString JsonWaveCurve;
-		(*JsonAnimatorObject)->TryGetStringField(GET_MEMBER_NAME_STRING_CHECKED(UPropertyAnimatorCurve, WaveCurve), JsonWaveCurve);
+		TSharedPtr<FPropertyAnimatorCorePresetObjectArchive> AnimatorArchive = InValue->AsMutableObject();
+
+		FString WaveCurveValue;
+		AnimatorArchive->Get(GET_MEMBER_NAME_STRING_CHECKED(UPropertyAnimatorCurve, WaveCurve), WaveCurveValue);
 
-		if (UPropertyAnimatorWaveCurve* Curve = LoadObject<UPropertyAnimatorWaveCurve>(nullptr, *JsonWaveCurve))
+		if (UPropertyAnimatorWaveCurve* Curve = LoadObject<UPropertyAnimatorWaveCurve>(nullptr, *WaveCurveValue))
 		{
 			SetWaveCurve(Curve);
 		}
 
-		bool bJsonEaseInEnabled = bEaseInEnabled;
-		(*JsonAnimatorObject)->TryGetBoolField(GET_MEMBER_NAME_STRING_CHECKED(UPropertyAnimatorCurve, bEaseInEnabled), bJsonEaseInEnabled);
-		SetEaseInEnabled(bJsonEaseInEnabled);
+		bool bEaseInEnabledValue = bEaseInEnabled;
+		AnimatorArchive->Get(GET_MEMBER_NAME_STRING_CHECKED(UPropertyAnimatorCurve, bEaseInEnabled), bEaseInEnabledValue);
+		SetEaseInEnabled(bEaseInEnabledValue);
 
-		bool bJsonEaseOutEnabled = bEaseOutEnabled;
-		(*JsonAnimatorObject)->TryGetBoolField(GET_MEMBER_NAME_STRING_CHECKED(UPropertyAnimatorCurve, bEaseOutEnabled), bJsonEaseOutEnabled);
-		SetEaseOutEnabled(bJsonEaseOutEnabled);
+		bool bEaseOutEnabledValue = bEaseOutEnabled;
+		AnimatorArchive->Get(GET_MEMBER_NAME_STRING_CHECKED(UPropertyAnimatorCurve, bEaseOutEnabled), bEaseOutEnabledValue);
+		SetEaseOutEnabled(bEaseOutEnabledValue);
 
-		const TSharedPtr<FJsonObject>* JsonEaseInObject;
-		(*JsonAnimatorObject)->TryGetObjectField(GET_MEMBER_NAME_STRING_CHECKED(UPropertyAnimatorCurve, EaseIn), JsonEaseInObject);
+		TSharedPtr<FPropertyAnimatorCorePresetArchive> EaseInArchive;
+		AnimatorArchive->Get(GET_MEMBER_NAME_STRING_CHECKED(UPropertyAnimatorCurve, EaseIn), EaseInArchive);
 
-		if (JsonEaseInObject)
+		if (const TSharedPtr<FPropertyAnimatorCorePresetObjectArchive> EaseInObject = EaseInArchive->AsMutableObject())
 		{
-			FString JsonEaseCurve;
-			(*JsonEaseInObject)->TryGetStringField(GET_MEMBER_NAME_STRING_CHECKED(FPropertyAnimatorCurveEasing, EaseCurve), JsonEaseCurve);
+			FString EaseCurvePath;
+			EaseInObject->Get(GET_MEMBER_NAME_STRING_CHECKED(FPropertyAnimatorCurveEasing, EaseCurve), EaseCurvePath);
 
-			double JsonEaseDuration;
-			(*JsonEaseInObject)->TryGetNumberField(GET_MEMBER_NAME_STRING_CHECKED(FPropertyAnimatorCurveEasing, EaseDuration), JsonEaseDuration);
+			double EaseDuration;
+			EaseInObject->Get(GET_MEMBER_NAME_STRING_CHECKED(FPropertyAnimatorCurveEasing, EaseDuration), EaseDuration);
 
-			if (UPropertyAnimatorEaseCurve* EaseCurve = LoadObject<UPropertyAnimatorEaseCurve>(nullptr, *JsonEaseCurve))
+			if (UPropertyAnimatorEaseCurve* EaseCurve = LoadObject<UPropertyAnimatorEaseCurve>(nullptr, *EaseCurvePath))
 			{
 				FPropertyAnimatorCurveEasing Easing;
 				Easing.EaseCurve = EaseCurve;
-				Easing.EaseDuration = JsonEaseDuration;
+				Easing.EaseDuration = EaseDuration;
 				SetEaseIn(Easing);
 			}
 		}
 
-		const TSharedPtr<FJsonObject>* JsonEaseOutObject;
-		(*JsonAnimatorObject)->TryGetObjectField(GET_MEMBER_NAME_STRING_CHECKED(UPropertyAnimatorCurve, EaseOut), JsonEaseOutObject);
+		TSharedPtr<FPropertyAnimatorCorePresetArchive> EaseOutArchive;
+		AnimatorArchive->Get(GET_MEMBER_NAME_STRING_CHECKED(UPropertyAnimatorCurve, EaseOut), EaseOutArchive);
 
-		if (JsonEaseOutObject)
+		if (const TSharedPtr<FPropertyAnimatorCorePresetObjectArchive> EaseOutObject = EaseOutArchive->AsMutableObject())
 		{
-			FString JsonEaseCurve;
-			(*JsonEaseOutObject)->TryGetStringField(GET_MEMBER_NAME_STRING_CHECKED(FPropertyAnimatorCurveEasing, EaseCurve), JsonEaseCurve);
+			FString EaseCurvePath;
+			EaseOutObject->Get(GET_MEMBER_NAME_STRING_CHECKED(FPropertyAnimatorCurveEasing, EaseCurve), EaseCurvePath);
 
-			double JsonEaseDuration;
-			(*JsonEaseOutObject)->TryGetNumberField(GET_MEMBER_NAME_STRING_CHECKED(FPropertyAnimatorCurveEasing, EaseDuration), JsonEaseDuration);
+			double EaseDuration;
+			EaseOutObject->Get(GET_MEMBER_NAME_STRING_CHECKED(FPropertyAnimatorCurveEasing, EaseDuration), EaseDuration);
 
-			if (UPropertyAnimatorEaseCurve* EaseCurve = LoadObject<UPropertyAnimatorEaseCurve>(nullptr, *JsonEaseCurve))
+			if (UPropertyAnimatorEaseCurve* EaseCurve = LoadObject<UPropertyAnimatorEaseCurve>(nullptr, *EaseCurvePath))
 			{
 				FPropertyAnimatorCurveEasing Easing;
 				Easing.EaseCurve = EaseCurve;
-				Easing.EaseDuration = JsonEaseDuration;
+				Easing.EaseDuration = EaseDuration;
 				SetEaseOut(Easing);
 			}
 		}
@@ -192,34 +190,34 @@ bool UPropertyAnimatorCurve::ImportPreset(const UPropertyAnimatorCorePresetBase*
 	return false;
 }
 
-bool UPropertyAnimatorCurve::ExportPreset(const UPropertyAnimatorCorePresetBase* InPreset, TSharedPtr<FJsonValue>& OutValue)
+bool UPropertyAnimatorCurve::ExportPreset(const UPropertyAnimatorCorePresetBase* InPreset, TSharedPtr<FPropertyAnimatorCorePresetArchive>& OutValue) const
 {
-	const TSharedPtr<FJsonObject>* JsonAnimatorObject;
-
-	if (Super::ExportPreset(InPreset, OutValue) && OutValue->TryGetObject(JsonAnimatorObject))
+	if (Super::ExportPreset(InPreset, OutValue) && OutValue->IsObject())
 	{
+		TSharedPtr<FPropertyAnimatorCorePresetObjectArchive> AnimatorArchive = OutValue->AsMutableObject();
+
 		if (WaveCurve)
 		{
-			(*JsonAnimatorObject)->SetStringField(GET_MEMBER_NAME_STRING_CHECKED(UPropertyAnimatorCurve, WaveCurve), WaveCurve.GetPath());
+			AnimatorArchive->Set(GET_MEMBER_NAME_STRING_CHECKED(UPropertyAnimatorCurve, WaveCurve), WaveCurve.GetPath());
 		}
 
-		(*JsonAnimatorObject)->SetBoolField(GET_MEMBER_NAME_STRING_CHECKED(UPropertyAnimatorCurve, bEaseInEnabled), bEaseInEnabled);
-		(*JsonAnimatorObject)->SetBoolField(GET_MEMBER_NAME_STRING_CHECKED(UPropertyAnimatorCurve, bEaseOutEnabled), bEaseOutEnabled);
+		AnimatorArchive->Set(GET_MEMBER_NAME_STRING_CHECKED(UPropertyAnimatorCurve, bEaseInEnabled), bEaseInEnabled);
+		AnimatorArchive->Set(GET_MEMBER_NAME_STRING_CHECKED(UPropertyAnimatorCurve, bEaseOutEnabled), bEaseOutEnabled);
 
 		if (EaseIn.EaseCurve)
 		{
-			TSharedRef<FJsonObject> JsonEaseObject = MakeShared<FJsonObject>();
-			JsonEaseObject->SetStringField(GET_MEMBER_NAME_STRING_CHECKED(FPropertyAnimatorCurveEasing, EaseCurve), EaseIn.EaseCurve.GetPath());
-			JsonEaseObject->SetNumberField(GET_MEMBER_NAME_STRING_CHECKED(FPropertyAnimatorCurveEasing, EaseDuration), EaseIn.EaseDuration);
-			(*JsonAnimatorObject)->SetObjectField(GET_MEMBER_NAME_STRING_CHECKED(UPropertyAnimatorCurve, EaseIn), JsonEaseObject);
+			TSharedRef<FPropertyAnimatorCorePresetObjectArchive> EaseObject = InPreset->GetArchiveImplementation()->CreateObject();
+			EaseObject->Set(GET_MEMBER_NAME_STRING_CHECKED(FPropertyAnimatorCurveEasing, EaseCurve), EaseIn.EaseCurve.GetPath());
+			EaseObject->Set(GET_MEMBER_NAME_STRING_CHECKED(FPropertyAnimatorCurveEasing, EaseDuration), EaseIn.EaseDuration);
+			AnimatorArchive->Set(GET_MEMBER_NAME_STRING_CHECKED(UPropertyAnimatorCurve, EaseIn), EaseObject);
 		}
 
 		if (EaseOut.EaseCurve)
 		{
-			TSharedRef<FJsonObject> JsonEaseObject = MakeShared<FJsonObject>();
-			JsonEaseObject->SetStringField(GET_MEMBER_NAME_STRING_CHECKED(FPropertyAnimatorCurveEasing, EaseCurve), EaseOut.EaseCurve.GetPath());
-			JsonEaseObject->SetNumberField(GET_MEMBER_NAME_STRING_CHECKED(FPropertyAnimatorCurveEasing, EaseDuration), EaseOut.EaseDuration);
-			(*JsonAnimatorObject)->SetObjectField(GET_MEMBER_NAME_STRING_CHECKED(UPropertyAnimatorCurve, EaseOut), JsonEaseObject);
+			TSharedRef<FPropertyAnimatorCorePresetObjectArchive> EaseObject = InPreset->GetArchiveImplementation()->CreateObject();
+			EaseObject->Set(GET_MEMBER_NAME_STRING_CHECKED(FPropertyAnimatorCurveEasing, EaseCurve), EaseOut.EaseCurve.GetPath());
+			EaseObject->Set(GET_MEMBER_NAME_STRING_CHECKED(FPropertyAnimatorCurveEasing, EaseDuration), EaseOut.EaseDuration);
+			AnimatorArchive->Set(GET_MEMBER_NAME_STRING_CHECKED(UPropertyAnimatorCurve, EaseOut), EaseObject);
 		}
 
 		return true;
diff --git a/Engine/Plugins/Experimental/PropertyAnimator/Source/PropertyAnimator/Private/Animators/PropertyAnimatorNumericBase.cpp b/Engine/Plugins/Experimental/PropertyAnimator/Source/PropertyAnimator/Private/Animators/PropertyAnimatorNumericBase.cpp
--- a/Engine/Plugins/Experimental/PropertyAnimator/Source/PropertyAnimator/Private/Animators/PropertyAnimatorNumericBase.cpp
+++ b/Engine/Plugins/Experimental/PropertyAnimator/Source/PropertyAnimator/Private/Animators/PropertyAnimatorNumericBase.cpp
@@ -2,8 +2,6 @@
 
 #include "Animators/PropertyAnimatorNumericBase.h"
 
-#include "Dom/JsonObject.h"
-#include "Dom/JsonValue.h"
 #include "Properties/PropertyAnimatorFloatContext.h"
 #include "Properties/PropertyAnimatorRotatorContext.h"
 #include "Properties/PropertyAnimatorVectorContext.h"
@@ -187,7 +185,7 @@ void UPropertyAnimatorNumericBase::EvaluateProperties(FInstancedPropertyBag& InP
 		{
 			if (FMath::Abs(PropertyTimeElapsed) > CycleDuration)
 			{
-				PropertyTimeElapsed = CycleDuration - UE_KINDA_SMALL_NUMBER;
+				PropertyTimeElapsed = CycleDuration;
 			}
 		}
 		else if (CycleMode == EPropertyAnimatorCycleMode::Loop)
@@ -196,7 +194,7 @@ void UPropertyAnimatorNumericBase::EvaluateProperties(FInstancedPropertyBag& InP
 
 			if (FMath::Abs(PropertyTimeElapsed) > CycleDuration)
 			{
-				PropertyTimeElapsed = CycleDuration - UE_KINDA_SMALL_NUMBER;
+				PropertyTimeElapsed = CycleDuration;
 			}
 		}
 		else if (CycleMode == EPropertyAnimatorCycleMode::PingPong)
@@ -206,7 +204,7 @@ void UPropertyAnimatorNumericBase::EvaluateProperties(FInstancedPropertyBag& InP
 
 			if (FMath::Abs(PropertyTimeElapsed) > CycleDuration)
 			{
-				PropertyTimeElapsed = CycleDuration - UE_KINDA_SMALL_NUMBER;
+				PropertyTimeElapsed = CycleDuration;
 			}
 
 			if (bReverse)
@@ -276,38 +274,38 @@ void UPropertyAnimatorNumericBase::OnAnimatorRegistered(FPropertyAnimatorCoreMet
 	InMetadata.Category = TEXT("Numeric");
 }
 
-bool UPropertyAnimatorNumericBase::ImportPreset(const UPropertyAnimatorCorePresetBase* InPreset, const TSharedRef<FJsonValue>& InValue)
+bool UPropertyAnimatorNumericBase::ImportPreset(const UPropertyAnimatorCorePresetBase* InPreset, const TSharedRef<FPropertyAnimatorCorePresetArchive>& InValue)
 {
-	const TSharedPtr<FJsonObject>* JsonAnimatorObject;
-
-	if (Super::ImportPreset(InPreset, InValue) && InValue->TryGetObject(JsonAnimatorObject))
+	if (Super::ImportPreset(InPreset, InValue) && InValue->IsObject())
 	{
-		double JsonMagnitude = Magnitude;
-		(*JsonAnimatorObject)->TryGetNumberField(GET_MEMBER_NAME_STRING_CHECKED(UPropertyAnimatorNumericBase, Magnitude), JsonMagnitude);
-		SetMagnitude(JsonMagnitude);
+		const TSharedPtr<FPropertyAnimatorCorePresetObjectArchive> AnimatorArchive = InValue->AsMutableObject();
+
+		double MagnitudeValue = Magnitude;
+		AnimatorArchive->Get(GET_MEMBER_NAME_STRING_CHECKED(UPropertyAnimatorNumericBase, Magnitude), MagnitudeValue);
+		SetMagnitude(MagnitudeValue);
 
 		if (CycleMode != EPropertyAnimatorCycleMode::None)
 		{
-			double JsonCycleMode = static_cast<double>(CycleMode);
-			(*JsonAnimatorObject)->TryGetNumberField(GET_MEMBER_NAME_STRING_CHECKED(UPropertyAnimatorNumericBase, CycleMode), JsonCycleMode);
-			SetCycleMode(static_cast<EPropertyAnimatorCycleMode>(JsonCycleMode));
+			uint64 CycleModeValue = static_cast<uint64>(CycleMode);
+			AnimatorArchive->Get(GET_MEMBER_NAME_STRING_CHECKED(UPropertyAnimatorNumericBase, CycleMode), CycleModeValue);
+			SetCycleMode(static_cast<EPropertyAnimatorCycleMode>(CycleModeValue));
 
-			double JsonCycleDuration = CycleDuration;
-			(*JsonAnimatorObject)->TryGetNumberField(GET_MEMBER_NAME_STRING_CHECKED(UPropertyAnimatorNumericBase, CycleDuration), JsonCycleDuration);
-			SetCycleDuration(JsonCycleDuration);
+			double CycleDurationValue = CycleDuration;
+			AnimatorArchive->Get(GET_MEMBER_NAME_STRING_CHECKED(UPropertyAnimatorNumericBase, CycleDuration), CycleDurationValue);
+			SetCycleDuration(CycleDurationValue);
 
-			double JsonCycleGapDuration = CycleGapDuration;
-			(*JsonAnimatorObject)->TryGetNumberField(GET_MEMBER_NAME_STRING_CHECKED(UPropertyAnimatorNumericBase, CycleGapDuration), JsonCycleGapDuration);
-			SetCycleGapDuration(JsonCycleGapDuration);
+			double CycleGapDurationValue = CycleGapDuration;
+			AnimatorArchive->Get(GET_MEMBER_NAME_STRING_CHECKED(UPropertyAnimatorNumericBase, CycleGapDuration), CycleGapDurationValue);
+			SetCycleGapDuration(CycleGapDurationValue);
 		}
 
-		bool bJsonRandomTimeOffset = bRandomTimeOffset;
-		(*JsonAnimatorObject)->TryGetBoolField(GET_MEMBER_NAME_STRING_CHECKED(UPropertyAnimatorNumericBase, bRandomTimeOffset), bJsonRandomTimeOffset);
-		SetRandomTimeOffset(bJsonRandomTimeOffset);
+		bool bRandomTimeOffsetValue = bRandomTimeOffset;
+		AnimatorArchive->Get(GET_MEMBER_NAME_STRING_CHECKED(UPropertyAnimatorNumericBase, bRandomTimeOffset), bRandomTimeOffsetValue);
+		SetRandomTimeOffset(bRandomTimeOffsetValue);
 
-		double JsonSeed = Seed;
-		(*JsonAnimatorObject)->TryGetNumberField(GET_MEMBER_NAME_STRING_CHECKED(UPropertyAnimatorNumericBase, Seed), JsonSeed);
-		SetSeed(JsonSeed);
+		int64 SeedValue = Seed;
+		AnimatorArchive->Get(GET_MEMBER_NAME_STRING_CHECKED(UPropertyAnimatorNumericBase, Seed), SeedValue);
+		SetSeed(SeedValue);
 
 		return true;
 	}
@@ -315,23 +313,23 @@ bool UPropertyAnimatorNumericBase::ImportPreset(const UPropertyAnimatorCorePrese
 	return false;
 }
 
-bool UPropertyAnimatorNumericBase::ExportPreset(const UPropertyAnimatorCorePresetBase* InPreset, TSharedPtr<FJsonValue>& OutValue)
+bool UPropertyAnimatorNumericBase::ExportPreset(const UPropertyAnimatorCorePresetBase* InPreset, TSharedPtr<FPropertyAnimatorCorePresetArchive>& OutValue) const
 {
-	const TSharedPtr<FJsonObject>* JsonAnimatorObject;
-
-	if (Super::ExportPreset(InPreset, OutValue) && OutValue->TryGetObject(JsonAnimatorObject))
+	if (Super::ExportPreset(InPreset, OutValue) && OutValue->IsObject())
 	{
-		(*JsonAnimatorObject)->SetNumberField(GET_MEMBER_NAME_STRING_CHECKED(UPropertyAnimatorNumericBase, Magnitude), Magnitude);
+		const TSharedPtr<FPropertyAnimatorCorePresetObjectArchive> AnimatorArchive = OutValue->AsMutableObject();
+
+		AnimatorArchive->Set(GET_MEMBER_NAME_STRING_CHECKED(UPropertyAnimatorNumericBase, Magnitude), Magnitude);
 
 		if (CycleMode != EPropertyAnimatorCycleMode::None)
 		{
-			(*JsonAnimatorObject)->SetNumberField(GET_MEMBER_NAME_STRING_CHECKED(UPropertyAnimatorNumericBase, CycleMode), static_cast<double>(CycleMode));
-			(*JsonAnimatorObject)->SetNumberField(GET_MEMBER_NAME_STRING_CHECKED(UPropertyAnimatorNumericBase, CycleDuration), CycleDuration);
-			(*JsonAnimatorObject)->SetNumberField(GET_MEMBER_NAME_STRING_CHECKED(UPropertyAnimatorNumericBase, CycleGapDuration), CycleGapDuration);
+			AnimatorArchive->Set(GET_MEMBER_NAME_STRING_CHECKED(UPropertyAnimatorNumericBase, CycleMode), static_cast<uint64>(CycleMode));
+			AnimatorArchive->Set(GET_MEMBER_NAME_STRING_CHECKED(UPropertyAnimatorNumericBase, CycleDuration), CycleDuration);
+			AnimatorArchive->Set(GET_MEMBER_NAME_STRING_CHECKED(UPropertyAnimatorNumericBase, CycleGapDuration), CycleGapDuration);
 		}
 
-		(*JsonAnimatorObject)->SetBoolField(GET_MEMBER_NAME_STRING_CHECKED(UPropertyAnimatorNumericBase, bRandomTimeOffset), bRandomTimeOffset);
-		(*JsonAnimatorObject)->SetNumberField(GET_MEMBER_NAME_STRING_CHECKED(UPropertyAnimatorNumericBase, Seed), Seed);
+		AnimatorArchive->Set(GET_MEMBER_NAME_STRING_CHECKED(UPropertyAnimatorNumericBase, bRandomTimeOffset), bRandomTimeOffset);
+		AnimatorArchive->Set(GET_MEMBER_NAME_STRING_CHECKED(UPropertyAnimatorNumericBase, Seed), static_cast<int64>(Seed));
 
 		return true;
 	}
diff --git a/Engine/Plugins/Experimental/PropertyAnimator/Source/PropertyAnimator/Private/Animators/PropertyAnimatorSoundWave.cpp b/Engine/Plugins/Experimental/PropertyAnimator/Source/PropertyAnimator/Private/Animators/PropertyAnimatorSoundWave.cpp
--- a/Engine/Plugins/Experimental/PropertyAnimator/Source/PropertyAnimator/Private/Animators/PropertyAnimatorSoundWave.cpp
+++ b/Engine/Plugins/Experimental/PropertyAnimator/Source/PropertyAnimator/Private/Animators/PropertyAnimatorSoundWave.cpp
@@ -2,8 +2,6 @@
 
 #include "Animators/PropertyAnimatorSoundWave.h"
 
-#include "Dom/JsonObject.h"
-#include "Dom/JsonValue.h"
 #include "LoudnessNRT.h"
 #include "Properties/Handlers/PropertyAnimatorCoreHandlerBase.h"
 #include "Properties/PropertyAnimatorFloatContext.h"
@@ -96,23 +94,23 @@ bool UPropertyAnimatorSoundWave::EvaluateProperty(const FPropertyAnimatorCoreDat
 	return false;
 }
 
-bool UPropertyAnimatorSoundWave::ImportPreset(const UPropertyAnimatorCorePresetBase* InPreset, const TSharedRef<FJsonValue>& InValue)
+bool UPropertyAnimatorSoundWave::ImportPreset(const UPropertyAnimatorCorePresetBase* InPreset, const TSharedRef<FPropertyAnimatorCorePresetArchive>& InValue)
 {
-	const TSharedPtr<FJsonObject>* JsonAnimatorObject;
-
-	if (Super::ImportPreset(InPreset, InValue) && InValue->TryGetObject(JsonAnimatorObject))
+	if (Super::ImportPreset(InPreset, InValue) && InValue->IsObject())
 	{
-		FString JsonSoundWave;
-		(*JsonAnimatorObject)->TryGetStringField(GET_MEMBER_NAME_STRING_CHECKED(UPropertyAnimatorSoundWave, SampledSoundWave), JsonSoundWave);
+		const TSharedPtr<FPropertyAnimatorCorePresetObjectArchive> AnimatorArchive = InValue->AsMutableObject();
+
+		FString SoundWaveValue;
+		AnimatorArchive->Get(GET_MEMBER_NAME_STRING_CHECKED(UPropertyAnimatorSoundWave, SampledSoundWave), SoundWaveValue);
 
-		if (USoundWave* SoundWave = LoadObject<USoundWave>(nullptr, *JsonSoundWave))
+		if (USoundWave* SoundWave = LoadObject<USoundWave>(nullptr, *SoundWaveValue))
 		{
 			SetSampledSoundWave(SoundWave);
 		}
 
-		bool bJsonLoop = bLoop;
-		(*JsonAnimatorObject)->TryGetBoolField(GET_MEMBER_NAME_STRING_CHECKED(UPropertyAnimatorSoundWave, bLoop), bLoop);
-		SetLoop(bJsonLoop);
+		bool bLoopValue = bLoop;
+		AnimatorArchive->Get(GET_MEMBER_NAME_STRING_CHECKED(UPropertyAnimatorSoundWave, bLoop), bLoop);
+		SetLoop(bLoopValue);
 
 		return true;
 	}
@@ -120,18 +118,18 @@ bool UPropertyAnimatorSoundWave::ImportPreset(const UPropertyAnimatorCorePresetB
 	return false;
 }
 
-bool UPropertyAnimatorSoundWave::ExportPreset(const UPropertyAnimatorCorePresetBase* InPreset, TSharedPtr<FJsonValue>& OutValue)
+bool UPropertyAnimatorSoundWave::ExportPreset(const UPropertyAnimatorCorePresetBase* InPreset, TSharedPtr<FPropertyAnimatorCorePresetArchive>& OutValue) const
 {
-	const TSharedPtr<FJsonObject>* JsonAnimatorObject;
-
-	if (Super::ExportPreset(InPreset, OutValue) && OutValue->TryGetObject(JsonAnimatorObject))
+	if (Super::ExportPreset(InPreset, OutValue) && OutValue->IsObject())
 	{
+		const TSharedPtr<FPropertyAnimatorCorePresetObjectArchive> AnimatorArchive = OutValue->AsMutableObject();
+
 		if (SampledSoundWave)
 		{
-			(*JsonAnimatorObject)->SetStringField(GET_MEMBER_NAME_STRING_CHECKED(UPropertyAnimatorSoundWave, SampledSoundWave), SampledSoundWave.GetPath());
+			AnimatorArchive->Set(GET_MEMBER_NAME_STRING_CHECKED(UPropertyAnimatorSoundWave, SampledSoundWave), SampledSoundWave.GetPath());
 		}
 
-		(*JsonAnimatorObject)->SetBoolField(GET_MEMBER_NAME_STRING_CHECKED(UPropertyAnimatorSoundWave, bLoop), bLoop);
+		AnimatorArchive->Set(GET_MEMBER_NAME_STRING_CHECKED(UPropertyAnimatorSoundWave, bLoop), bLoop);
 
 		return true;
 	}
diff --git a/Engine/Plugins/Experimental/PropertyAnimator/Source/PropertyAnimator/Private/Properties/PropertyAnimatorFloatContext.cpp b/Engine/Plugins/Experimental/PropertyAnimator/Source/PropertyAnimator/Private/Properties/PropertyAnimatorFloatContext.cpp
--- a/Engine/Plugins/Experimental/PropertyAnimator/Source/PropertyAnimator/Private/Properties/PropertyAnimatorFloatContext.cpp
+++ b/Engine/Plugins/Experimental/PropertyAnimator/Source/PropertyAnimator/Private/Properties/PropertyAnimatorFloatContext.cpp
@@ -3,8 +3,6 @@
 #include "Properties/PropertyAnimatorFloatContext.h"
 
 #include "Animators/PropertyAnimatorCoreBase.h"
-#include "Dom/JsonObject.h"
-#include "Dom/JsonValue.h"
 
 void UPropertyAnimatorFloatContext::SetAmplitudeMin(double InAmplitude)
 {
@@ -89,19 +87,19 @@ void UPropertyAnimatorFloatContext::OnAnimatedPropertyLinked()
 #endif
 }
 
-bool UPropertyAnimatorFloatContext::ImportPreset(const UPropertyAnimatorCorePresetBase* InPreset, const TSharedRef<FJsonValue>& InValue)
+bool UPropertyAnimatorFloatContext::ImportPreset(const UPropertyAnimatorCorePresetBase* InPreset, const TSharedRef<FPropertyAnimatorCorePresetArchive>& InValue)
 {
-	TSharedPtr<FJsonObject>* JsonObject = nullptr;
-
-	if (Super::ImportPreset(InPreset, InValue) && InValue->TryGetObject(JsonObject))
+	if (Super::ImportPreset(InPreset, InValue) && InValue->IsObject())
 	{
-		double JsonAmplitudeMin = AmplitudeMin;
-		(*JsonObject)->TryGetNumberField(GET_MEMBER_NAME_STRING_CHECKED(UPropertyAnimatorFloatContext, AmplitudeMin), JsonAmplitudeMin);
-		SetAmplitudeMin(JsonAmplitudeMin);
+		const TSharedPtr<FPropertyAnimatorCorePresetObjectArchive> ContextArchive = InValue->AsMutableObject();
+
+		double AmplitudeMinValue = AmplitudeMin;
+		ContextArchive->Get(GET_MEMBER_NAME_STRING_CHECKED(UPropertyAnimatorFloatContext, AmplitudeMin), AmplitudeMinValue);
+		SetAmplitudeMin(AmplitudeMinValue);
 
-		double JsonAmplitudeMax = AmplitudeMax;
-		(*JsonObject)->TryGetNumberField(GET_MEMBER_NAME_STRING_CHECKED(UPropertyAnimatorFloatContext, AmplitudeMax), JsonAmplitudeMax);
-		SetAmplitudeMax(JsonAmplitudeMax);
+		double AmplitudeMaxValue = AmplitudeMax;
+		ContextArchive->Get(GET_MEMBER_NAME_STRING_CHECKED(UPropertyAnimatorFloatContext, AmplitudeMax), AmplitudeMaxValue);
+		SetAmplitudeMax(AmplitudeMaxValue);
 
 		return true;
 	}
@@ -109,14 +107,14 @@ bool UPropertyAnimatorFloatContext::ImportPreset(const UPropertyAnimatorCorePres
 	return false;
 }
 
-bool UPropertyAnimatorFloatContext::ExportPreset(const UPropertyAnimatorCorePresetBase* InPreset, TSharedPtr<FJsonValue>& OutValue)
+bool UPropertyAnimatorFloatContext::ExportPreset(const UPropertyAnimatorCorePresetBase* InPreset, TSharedPtr<FPropertyAnimatorCorePresetArchive>& OutValue) const
 {
-	TSharedPtr<FJsonObject>* JsonObject = nullptr;
-
-	if (Super::ExportPreset(InPreset, OutValue) && OutValue->TryGetObject(JsonObject))
+	if (Super::ExportPreset(InPreset, OutValue) && OutValue->IsObject())
 	{
-		(*JsonObject)->SetNumberField(GET_MEMBER_NAME_STRING_CHECKED(UPropertyAnimatorFloatContext, AmplitudeMin), AmplitudeMin);
-		(*JsonObject)->SetNumberField(GET_MEMBER_NAME_STRING_CHECKED(UPropertyAnimatorFloatContext, AmplitudeMax), AmplitudeMax);
+		const TSharedPtr<FPropertyAnimatorCorePresetObjectArchive> ContextArchive = OutValue->AsMutableObject();
+
+		ContextArchive->Set(GET_MEMBER_NAME_STRING_CHECKED(UPropertyAnimatorFloatContext, AmplitudeMin), AmplitudeMin);
+		ContextArchive->Set(GET_MEMBER_NAME_STRING_CHECKED(UPropertyAnimatorFloatContext, AmplitudeMax), AmplitudeMax);
 
 		return true;
 	}
diff --git a/Engine/Plugins/Experimental/PropertyAnimator/Source/PropertyAnimator/Private/Properties/PropertyAnimatorRotatorContext.cpp b/Engine/Plugins/Experimental/PropertyAnimator/Source/PropertyAnimator/Private/Properties/PropertyAnimatorRotatorContext.cpp
--- a/Engine/Plugins/Experimental/PropertyAnimator/Source/PropertyAnimator/Private/Properties/PropertyAnimatorRotatorContext.cpp
+++ b/Engine/Plugins/Experimental/PropertyAnimator/Source/PropertyAnimator/Private/Properties/PropertyAnimatorRotatorContext.cpp
@@ -3,8 +3,6 @@
 #include "Properties/PropertyAnimatorRotatorContext.h"
 
 #include "Animators/PropertyAnimatorCoreBase.h"
-#include "Dom/JsonObject.h"
-#include "Dom/JsonValue.h"
 
 void UPropertyAnimatorRotatorContext::SetAmplitudeMin(const FRotator& InAmplitude)
 {
@@ -89,22 +87,22 @@ void UPropertyAnimatorRotatorContext::OnAnimatedPropertyLinked()
 #endif
 }
 
-bool UPropertyAnimatorRotatorContext::ImportPreset(const UPropertyAnimatorCorePresetBase* InPreset, const TSharedRef<FJsonValue>& InValue)
+bool UPropertyAnimatorRotatorContext::ImportPreset(const UPropertyAnimatorCorePresetBase* InPreset, const TSharedRef<FPropertyAnimatorCorePresetArchive>& InValue)
 {
-	TSharedPtr<FJsonObject>* JsonObject = nullptr;
-
-	if (Super::ImportPreset(InPreset, InValue) && InValue->TryGetObject(JsonObject))
+	if (Super::ImportPreset(InPreset, InValue) && InValue->IsObject())
 	{
-		FString JsonAmplitudeMin = AmplitudeMin.ToString();
-		(*JsonObject)->TryGetStringField(GET_MEMBER_NAME_STRING_CHECKED(UPropertyAnimatorRotatorContext, AmplitudeMin), JsonAmplitudeMin);
+		const TSharedPtr<FPropertyAnimatorCorePresetObjectArchive> ContextArchive = InValue->AsMutableObject();
+
+		FString AmplitudeMinValue = AmplitudeMin.ToString();
+		ContextArchive->Get(GET_MEMBER_NAME_STRING_CHECKED(UPropertyAnimatorRotatorContext, AmplitudeMin), AmplitudeMinValue);
 		FRotator ParseAmplitudeMin;
-		ParseAmplitudeMin.InitFromString(JsonAmplitudeMin);
+		ParseAmplitudeMin.InitFromString(AmplitudeMinValue);
 		SetAmplitudeMin(ParseAmplitudeMin);
 
-		FString JsonAmplitudeMax = AmplitudeMax.ToString();
-		(*JsonObject)->TryGetStringField(GET_MEMBER_NAME_STRING_CHECKED(UPropertyAnimatorRotatorContext, AmplitudeMax), JsonAmplitudeMax);
+		FString AmplitudeMaxValue = AmplitudeMax.ToString();
+		ContextArchive->Get(GET_MEMBER_NAME_STRING_CHECKED(UPropertyAnimatorRotatorContext, AmplitudeMax), AmplitudeMaxValue);
 		FRotator ParseAmplitudeMax;
-		ParseAmplitudeMax.InitFromString(JsonAmplitudeMax);
+		ParseAmplitudeMax.InitFromString(AmplitudeMaxValue);
 		SetAmplitudeMax(ParseAmplitudeMax);
 
 		return true;
@@ -113,14 +111,14 @@ bool UPropertyAnimatorRotatorContext::ImportPreset(const UPropertyAnimatorCorePr
 	return false;
 }
 
-bool UPropertyAnimatorRotatorContext::ExportPreset(const UPropertyAnimatorCorePresetBase* InPreset, TSharedPtr<FJsonValue>& OutValue)
+bool UPropertyAnimatorRotatorContext::ExportPreset(const UPropertyAnimatorCorePresetBase* InPreset, TSharedPtr<FPropertyAnimatorCorePresetArchive>& OutValue) const
 {
-	TSharedPtr<FJsonObject>* JsonObject = nullptr;
-
-	if (Super::ExportPreset(InPreset, OutValue) && OutValue->TryGetObject(JsonObject))
+	if (Super::ExportPreset(InPreset, OutValue) && OutValue->IsObject())
 	{
-		(*JsonObject)->SetStringField(GET_MEMBER_NAME_STRING_CHECKED(UPropertyAnimatorRotatorContext, AmplitudeMin), AmplitudeMin.ToString());
-		(*JsonObject)->SetStringField(GET_MEMBER_NAME_STRING_CHECKED(UPropertyAnimatorRotatorContext, AmplitudeMax), AmplitudeMax.ToString());
+		const TSharedPtr<FPropertyAnimatorCorePresetObjectArchive> ContextArchive = OutValue->AsMutableObject();
+
+		ContextArchive->Set(GET_MEMBER_NAME_STRING_CHECKED(UPropertyAnimatorRotatorContext, AmplitudeMin), AmplitudeMin.ToString());
+		ContextArchive->Set(GET_MEMBER_NAME_STRING_CHECKED(UPropertyAnimatorRotatorContext, AmplitudeMax), AmplitudeMax.ToString());
 
 		return true;
 	}
diff --git a/Engine/Plugins/Experimental/PropertyAnimator/Source/PropertyAnimator/Private/Properties/PropertyAnimatorTextResolver.cpp b/Engine/Plugins/Experimental/PropertyAnimator/Source/PropertyAnimator/Private/Properties/PropertyAnimatorTextResolver.cpp
--- a/Engine/Plugins/Experimental/PropertyAnimator/Source/PropertyAnimator/Private/Properties/PropertyAnimatorTextResolver.cpp
+++ b/Engine/Plugins/Experimental/PropertyAnimator/Source/PropertyAnimator/Private/Properties/PropertyAnimatorTextResolver.cpp
@@ -2,8 +2,7 @@
 
 #include "Properties/PropertyAnimatorTextResolver.h"
 
-#include "Dom/JsonObject.h"
-#include "Dom/JsonValue.h"
+#include "Presets/PropertyAnimatorCorePresetArchive.h"
 #include "Text3DComponent.h"
 
 void UPropertyAnimatorTextResolver::SetUnit(EPropertyAnimatorTextResolverRangeUnit InUnit)
@@ -260,57 +259,55 @@ void UPropertyAnimatorTextResolver::ResolveProperties(const FPropertyAnimatorCor
 	}
 }
 
-bool UPropertyAnimatorTextResolver::ImportPreset(const UPropertyAnimatorCorePresetBase* InPreset, const TSharedRef<FJsonValue>& InValue)
+bool UPropertyAnimatorTextResolver::ImportPreset(const UPropertyAnimatorCorePresetBase* InPreset, const TSharedRef<FPropertyAnimatorCorePresetArchive>& InValue)
 {
-	const TSharedPtr<FJsonObject>* JsonResolverObjectPtr;
-
-	if (Super::ImportPreset(InPreset, InValue) && InValue->TryGetObject(JsonResolverObjectPtr))
+	if (Super::ImportPreset(InPreset, InValue) && InValue->IsObject())
 	{
-		const TSharedPtr<FJsonObject> JsonResolverObject = *JsonResolverObjectPtr;
+		const TSharedPtr<FPropertyAnimatorCorePresetObjectArchive> ResolverArchive = InValue->AsMutableObject();
 
-		uint8 JsonUnit = static_cast<uint8>(Unit);
-		JsonResolverObject->TryGetNumberField(GET_MEMBER_NAME_STRING_CHECKED(UPropertyAnimatorTextResolver, Unit), JsonUnit);
-		SetUnit(static_cast<EPropertyAnimatorTextResolverRangeUnit>(JsonUnit));
+		uint64 UnitValue = static_cast<uint64>(Unit);
+		ResolverArchive->Get(GET_MEMBER_NAME_STRING_CHECKED(UPropertyAnimatorTextResolver, Unit), UnitValue);
+		SetUnit(static_cast<EPropertyAnimatorTextResolverRangeUnit>(UnitValue));
 
-		float JsonStart = Start;
-		JsonResolverObject->TryGetNumberField(GET_MEMBER_NAME_STRING_CHECKED(UPropertyAnimatorTextResolver, Start), JsonStart);
-		SetStart(JsonStart);
+		double StartValue = Start;
+		ResolverArchive->Get(GET_MEMBER_NAME_STRING_CHECKED(UPropertyAnimatorTextResolver, Start), StartValue);
+		SetStart(StartValue);
 
-		float JsonEnd = End;
-		JsonResolverObject->TryGetNumberField(GET_MEMBER_NAME_STRING_CHECKED(UPropertyAnimatorTextResolver, End), JsonEnd);
-		SetEnd(JsonEnd);
+		double EndValue = End;
+		ResolverArchive->Get(GET_MEMBER_NAME_STRING_CHECKED(UPropertyAnimatorTextResolver, End), EndValue);
+		SetEnd(EndValue);
 
-		float JsonOffset = Offset;
-		JsonResolverObject->TryGetNumberField(GET_MEMBER_NAME_STRING_CHECKED(UPropertyAnimatorTextResolver, Offset), JsonOffset);
-		SetOffset(JsonOffset);
+		double OffsetValue = Offset;
+		ResolverArchive->Get(GET_MEMBER_NAME_STRING_CHECKED(UPropertyAnimatorTextResolver, Offset), OffsetValue);
+		SetOffset(OffsetValue);
 
-		int32 JsonCharStartIndex = CharacterStartIndex;
-		JsonResolverObject->TryGetNumberField(GET_MEMBER_NAME_STRING_CHECKED(UPropertyAnimatorTextResolver, CharacterStartIndex), JsonCharStartIndex);
-		SetCharacterStartIndex(JsonCharStartIndex);
+		int64 CharStartIndexValue = CharacterStartIndex;
+		ResolverArchive->Get(GET_MEMBER_NAME_STRING_CHECKED(UPropertyAnimatorTextResolver, CharacterStartIndex), CharStartIndexValue);
+		SetCharacterStartIndex(CharStartIndexValue);
 
-		int32 JsonCharEndIndex = CharacterEndIndex;
-		JsonResolverObject->TryGetNumberField(GET_MEMBER_NAME_STRING_CHECKED(UPropertyAnimatorTextResolver, CharacterEndIndex), JsonCharEndIndex);
-		SetCharacterEndIndex(JsonCharEndIndex);
+		int64 CharEndIndexValue = CharacterEndIndex;
+		ResolverArchive->Get(GET_MEMBER_NAME_STRING_CHECKED(UPropertyAnimatorTextResolver, CharacterEndIndex), CharEndIndexValue);
+		SetCharacterEndIndex(CharEndIndexValue);
 
-		int32 JsonCharOffsetIndex = CharacterOffsetIndex;
-		JsonResolverObject->TryGetNumberField(GET_MEMBER_NAME_STRING_CHECKED(UPropertyAnimatorTextResolver, CharacterOffsetIndex), JsonCharOffsetIndex);
-		SetCharacterOffsetIndex(JsonCharOffsetIndex);
+		int64 CharOffsetIndexValue = CharacterOffsetIndex;
+		ResolverArchive->Get(GET_MEMBER_NAME_STRING_CHECKED(UPropertyAnimatorTextResolver, CharacterOffsetIndex), CharOffsetIndexValue);
+		SetCharacterOffsetIndex(CharOffsetIndexValue);
 
-		int32 JsonWordStartIndex = WordStartIndex;
-		JsonResolverObject->TryGetNumberField(GET_MEMBER_NAME_STRING_CHECKED(UPropertyAnimatorTextResolver, WordStartIndex), JsonWordStartIndex);
-		SetWordStartIndex(JsonWordStartIndex);
+		int64 WordStartIndexValue = WordStartIndex;
+		ResolverArchive->Get(GET_MEMBER_NAME_STRING_CHECKED(UPropertyAnimatorTextResolver, WordStartIndex), WordStartIndexValue);
+		SetWordStartIndex(WordStartIndexValue);
 
-		int32 JsonWordEndIndex = WordEndIndex;
-		JsonResolverObject->TryGetNumberField(GET_MEMBER_NAME_STRING_CHECKED(UPropertyAnimatorTextResolver, WordEndIndex), JsonWordEndIndex);
-		SetWordEndIndex(JsonWordEndIndex);
+		int64 WordEndIndexValue = WordEndIndex;
+		ResolverArchive->Get(GET_MEMBER_NAME_STRING_CHECKED(UPropertyAnimatorTextResolver, WordEndIndex), WordEndIndexValue);
+		SetWordEndIndex(WordEndIndexValue);
 
-		int32 JsonWordOffsetIndex = WordOffsetIndex;
-		JsonResolverObject->TryGetNumberField(GET_MEMBER_NAME_STRING_CHECKED(UPropertyAnimatorTextResolver, WordOffsetIndex), JsonWordOffsetIndex);
-		SetWordOffsetIndex(JsonWordOffsetIndex);
+		int64 WordOffsetIndexValue = WordOffsetIndex;
+		ResolverArchive->Get(GET_MEMBER_NAME_STRING_CHECKED(UPropertyAnimatorTextResolver, WordOffsetIndex), WordOffsetIndexValue);
+		SetWordOffsetIndex(WordOffsetIndexValue);
 
-		uint8 JsonDirection = static_cast<uint8>(Direction);
-		JsonResolverObject->TryGetNumberField(GET_MEMBER_NAME_STRING_CHECKED(UPropertyAnimatorTextResolver, Direction), JsonDirection);
-		SetDirection(static_cast<EPropertyAnimatorTextResolverRangeDirection>(JsonWordOffsetIndex));
+		uint64 DirectionValue = static_cast<uint64>(Direction);
+		ResolverArchive->Get(GET_MEMBER_NAME_STRING_CHECKED(UPropertyAnimatorTextResolver, Direction), DirectionValue);
+		SetDirection(static_cast<EPropertyAnimatorTextResolverRangeDirection>(WordOffsetIndexValue));
 
 		return true;
 	}
@@ -318,25 +315,23 @@ bool UPropertyAnimatorTextResolver::ImportPreset(const UPropertyAnimatorCorePres
 	return false;
 }
 
-bool UPropertyAnimatorTextResolver::ExportPreset(const UPropertyAnimatorCorePresetBase* InPreset, TSharedPtr<FJsonValue>& OutValue)
+bool UPropertyAnimatorTextResolver::ExportPreset(const UPropertyAnimatorCorePresetBase* InPreset, TSharedPtr<FPropertyAnimatorCorePresetArchive>& OutValue) const
 {
-	const TSharedPtr<FJsonObject>* JsonResolverObjectPtr;
-
-	if (Super::ExportPreset(InPreset, OutValue) && OutValue->TryGetObject(JsonResolverObjectPtr))
+	if (Super::ExportPreset(InPreset, OutValue) && OutValue->IsObject())
 	{
-		const TSharedPtr<FJsonObject> JsonResolverObject = *JsonResolverObjectPtr;
-
-		JsonResolverObject->SetNumberField(GET_MEMBER_NAME_STRING_CHECKED(UPropertyAnimatorTextResolver, Unit), static_cast<uint8>(Unit));
-		JsonResolverObject->SetNumberField(GET_MEMBER_NAME_STRING_CHECKED(UPropertyAnimatorTextResolver, Start), Start);
-		JsonResolverObject->SetNumberField(GET_MEMBER_NAME_STRING_CHECKED(UPropertyAnimatorTextResolver, End), End);
-		JsonResolverObject->SetNumberField(GET_MEMBER_NAME_STRING_CHECKED(UPropertyAnimatorTextResolver, Offset), Offset);
-		JsonResolverObject->SetNumberField(GET_MEMBER_NAME_STRING_CHECKED(UPropertyAnimatorTextResolver, CharacterStartIndex), CharacterStartIndex);
-		JsonResolverObject->SetNumberField(GET_MEMBER_NAME_STRING_CHECKED(UPropertyAnimatorTextResolver, CharacterEndIndex), CharacterEndIndex);
-		JsonResolverObject->SetNumberField(GET_MEMBER_NAME_STRING_CHECKED(UPropertyAnimatorTextResolver, CharacterOffsetIndex), CharacterOffsetIndex);
-		JsonResolverObject->SetNumberField(GET_MEMBER_NAME_STRING_CHECKED(UPropertyAnimatorTextResolver, WordStartIndex), WordStartIndex);
-		JsonResolverObject->SetNumberField(GET_MEMBER_NAME_STRING_CHECKED(UPropertyAnimatorTextResolver, WordEndIndex), WordEndIndex);
-		JsonResolverObject->SetNumberField(GET_MEMBER_NAME_STRING_CHECKED(UPropertyAnimatorTextResolver, WordOffsetIndex), WordOffsetIndex);
-		JsonResolverObject->SetNumberField(GET_MEMBER_NAME_STRING_CHECKED(UPropertyAnimatorTextResolver, Direction), static_cast<uint8>(Direction));
+		const TSharedPtr<FPropertyAnimatorCorePresetObjectArchive> ResolverArchive = OutValue->AsMutableObject();
+
+		ResolverArchive->Set(GET_MEMBER_NAME_STRING_CHECKED(UPropertyAnimatorTextResolver, Unit), static_cast<uint64>(Unit));
+		ResolverArchive->Set(GET_MEMBER_NAME_STRING_CHECKED(UPropertyAnimatorTextResolver, Start), Start);
+		ResolverArchive->Set(GET_MEMBER_NAME_STRING_CHECKED(UPropertyAnimatorTextResolver, End), End);
+		ResolverArchive->Set(GET_MEMBER_NAME_STRING_CHECKED(UPropertyAnimatorTextResolver, Offset), Offset);
+		ResolverArchive->Set(GET_MEMBER_NAME_STRING_CHECKED(UPropertyAnimatorTextResolver, CharacterStartIndex), static_cast<int64>(CharacterStartIndex));
+		ResolverArchive->Set(GET_MEMBER_NAME_STRING_CHECKED(UPropertyAnimatorTextResolver, CharacterEndIndex), static_cast<int64>(CharacterEndIndex));
+		ResolverArchive->Set(GET_MEMBER_NAME_STRING_CHECKED(UPropertyAnimatorTextResolver, CharacterOffsetIndex), static_cast<int64>(CharacterOffsetIndex));
+		ResolverArchive->Set(GET_MEMBER_NAME_STRING_CHECKED(UPropertyAnimatorTextResolver, WordStartIndex), static_cast<int64>(WordStartIndex));
+		ResolverArchive->Set(GET_MEMBER_NAME_STRING_CHECKED(UPropertyAnimatorTextResolver, WordEndIndex), static_cast<int64>(WordEndIndex));
+		ResolverArchive->Set(GET_MEMBER_NAME_STRING_CHECKED(UPropertyAnimatorTextResolver, WordOffsetIndex), static_cast<int64>(WordOffsetIndex));
+		ResolverArchive->Set(GET_MEMBER_NAME_STRING_CHECKED(UPropertyAnimatorTextResolver, Direction), static_cast<uint64>(Direction));
 
 		return true;
 	}
diff --git a/Engine/Plugins/Experimental/PropertyAnimator/Source/PropertyAnimator/Private/Properties/PropertyAnimatorTextResolver.h b/Engine/Plugins/Experimental/PropertyAnimator/Source/PropertyAnimator/Private/Properties/PropertyAnimatorTextResolver.h
--- a/Engine/Plugins/Experimental/PropertyAnimator/Source/PropertyAnimator/Private/Properties/PropertyAnimatorTextResolver.h
+++ b/Engine/Plugins/Experimental/PropertyAnimator/Source/PropertyAnimator/Private/Properties/PropertyAnimatorTextResolver.h
@@ -99,7 +99,7 @@ public:
 	}
 
 	PROPERTYANIMATOR_API void SetDirection(EPropertyAnimatorTextResolverRangeDirection InDirection);
-	EPropertyAnimatorTextResolverRangeDirection GeDirection() const
+	EPropertyAnimatorTextResolverRangeDirection GetDirection() const
 	{
 		return Direction;
 	}
@@ -107,12 +107,12 @@ public:
 	//~ Begin UPropertyAnimatorCoreResolver
 	virtual void GetResolvableProperties(const FPropertyAnimatorCoreData& InParentProperty, TSet<FPropertyAnimatorCoreData>& OutProperties) override;
 	virtual void ResolveProperties(const FPropertyAnimatorCoreData& InTemplateProperty, TArray<FPropertyAnimatorCoreData>& OutProperties, bool bInForEvaluation) override;
-	virtual bool ImportPreset(const UPropertyAnimatorCorePresetBase* InPreset, const TSharedRef<FJsonValue>& InValue) override;
-	virtual bool ExportPreset(const UPropertyAnimatorCorePresetBase* InPreset, TSharedPtr<FJsonValue>& OutValue) override;
+	virtual bool ImportPreset(const UPropertyAnimatorCorePresetBase* InPreset, const TSharedRef<FPropertyAnimatorCorePresetArchive>& InValue) override;
+	virtual bool ExportPreset(const UPropertyAnimatorCorePresetBase* InPreset, TSharedPtr<FPropertyAnimatorCorePresetArchive>& OutValue) const override;
 	//~ End UPropertyAnimatorCoreResolver
 
 protected:
-	UPROPERTY(EditInstanceOnly, Category="Animator")
+	UPROPERTY(EditInstanceOnly, Setter, Getter, Category="Animator")
 	EPropertyAnimatorTextResolverRangeUnit Unit = EPropertyAnimatorTextResolverRangeUnit::Percentage;
 
 	UPROPERTY(EditInstanceOnly, Setter, Getter, Category="Animator", meta=(ClampMin="0", ClampMax="100", Units=Percent, EditCondition="Unit == EPropertyAnimatorTextResolverRangeUnit::Percentage", EditConditionHides))
@@ -124,24 +124,24 @@ protected:
 	UPROPERTY(EditInstanceOnly, Setter, Getter, Category="Animator", meta=(Units=Percent, Delta="1", EditCondition="Unit == EPropertyAnimatorTextResolverRangeUnit::Percentage", EditConditionHides))
 	float Offset = 0.f;
 
-	UPROPERTY(EditInstanceOnly, Category="Animator", meta=(ClampMin="0", EditCondition="Unit == EPropertyAnimatorTextResolverRangeUnit::Character", EditConditionHides))
+	UPROPERTY(EditInstanceOnly, Setter, Getter, Category="Animator", meta=(ClampMin="0", EditCondition="Unit == EPropertyAnimatorTextResolverRangeUnit::Character", EditConditionHides))
 	int32 CharacterStartIndex = 0;
 
-	UPROPERTY(EditInstanceOnly, Category="Animator", meta=(ClampMin="0", EditCondition="Unit == EPropertyAnimatorTextResolverRangeUnit::Character", EditConditionHides))
+	UPROPERTY(EditInstanceOnly, Setter, Getter, Category="Animator", meta=(ClampMin="0", EditCondition="Unit == EPropertyAnimatorTextResolverRangeUnit::Character", EditConditionHides))
 	int32 CharacterEndIndex = 100;
 
-	UPROPERTY(EditInstanceOnly, Category="Animator", meta=(EditCondition="Unit == EPropertyAnimatorTextResolverRangeUnit::Character", EditConditionHides))
+	UPROPERTY(EditInstanceOnly, Setter, Getter, Category="Animator", meta=(EditCondition="Unit == EPropertyAnimatorTextResolverRangeUnit::Character", EditConditionHides))
 	int32 CharacterOffsetIndex = 0;
 
-	UPROPERTY(EditInstanceOnly, Category="Animator", meta=(ClampMin="0", EditCondition="Unit == EPropertyAnimatorTextResolverRangeUnit::Word", EditConditionHides))
+	UPROPERTY(EditInstanceOnly, Setter, Getter, Category="Animator", meta=(ClampMin="0", EditCondition="Unit == EPropertyAnimatorTextResolverRangeUnit::Word", EditConditionHides))
 	int32 WordStartIndex = 0;
 
-	UPROPERTY(EditInstanceOnly, Category="Animator", meta=(ClampMin="0", EditCondition="Unit == EPropertyAnimatorTextResolverRangeUnit::Word", EditConditionHides))
+	UPROPERTY(EditInstanceOnly, Setter, Getter, Category="Animator", meta=(ClampMin="0", EditCondition="Unit == EPropertyAnimatorTextResolverRangeUnit::Word", EditConditionHides))
 	int32 WordEndIndex = 100;
 
-	UPROPERTY(EditInstanceOnly, Category="Animator", meta=(EditCondition="Unit == EPropertyAnimatorTextResolverRangeUnit::Word", EditConditionHides))
+	UPROPERTY(EditInstanceOnly, Setter, Getter, Category="Animator", meta=(EditCondition="Unit == EPropertyAnimatorTextResolverRangeUnit::Word", EditConditionHides))
 	int32 WordOffsetIndex = 0;
 
-	UPROPERTY(EditInstanceOnly, Category="Animator")
+	UPROPERTY(EditInstanceOnly, Setter, Getter, Category="Animator")
 	EPropertyAnimatorTextResolverRangeDirection Direction = EPropertyAnimatorTextResolverRangeDirection::LeftToRight;
 };
\ No newline at end of file
diff --git a/Engine/Plugins/Experimental/PropertyAnimator/Source/PropertyAnimator/Private/Properties/PropertyAnimatorVectorContext.cpp b/Engine/Plugins/Experimental/PropertyAnimator/Source/PropertyAnimator/Private/Properties/PropertyAnimatorVectorContext.cpp
--- a/Engine/Plugins/Experimental/PropertyAnimator/Source/PropertyAnimator/Private/Properties/PropertyAnimatorVectorContext.cpp
+++ b/Engine/Plugins/Experimental/PropertyAnimator/Source/PropertyAnimator/Private/Properties/PropertyAnimatorVectorContext.cpp
@@ -3,8 +3,6 @@
 #include "Properties/PropertyAnimatorVectorContext.h"
 
 #include "Animators/PropertyAnimatorCoreBase.h"
-#include "Dom/JsonObject.h"
-#include "Dom/JsonValue.h"
 
 void UPropertyAnimatorVectorContext::SetAmplitudeMin(const FVector& InAmplitude)
 {
@@ -89,22 +87,22 @@ void UPropertyAnimatorVectorContext::OnAnimatedPropertyLinked()
 #endif
 }
 
-bool UPropertyAnimatorVectorContext::ImportPreset(const UPropertyAnimatorCorePresetBase* InPreset, const TSharedRef<FJsonValue>& InValue)
+bool UPropertyAnimatorVectorContext::ImportPreset(const UPropertyAnimatorCorePresetBase* InPreset, const TSharedRef<FPropertyAnimatorCorePresetArchive>& InValue)
 {
-	TSharedPtr<FJsonObject>* JsonObject = nullptr;
-
-	if (Super::ImportPreset(InPreset, InValue) && InValue->TryGetObject(JsonObject))
+	if (Super::ImportPreset(InPreset, InValue) && InValue->IsObject())
 	{
-		FString JsonAmplitudeMin = AmplitudeMin.ToString();
-		(*JsonObject)->TryGetStringField(GET_MEMBER_NAME_STRING_CHECKED(UPropertyAnimatorVectorContext, AmplitudeMin), JsonAmplitudeMin);
+		const TSharedPtr<FPropertyAnimatorCorePresetObjectArchive> ContextArchive = InValue->AsMutableObject();
+
+		FString AmplitudeMinValue = AmplitudeMin.ToString();
+		ContextArchive->Get(GET_MEMBER_NAME_STRING_CHECKED(UPropertyAnimatorVectorContext, AmplitudeMin), AmplitudeMinValue);
 		FVector ParseAmplitudeMin;
-		ParseAmplitudeMin.InitFromString(JsonAmplitudeMin);
+		ParseAmplitudeMin.InitFromString(AmplitudeMinValue);
 		SetAmplitudeMin(ParseAmplitudeMin);
 
-		FString JsonAmplitudeMax = AmplitudeMax.ToString();
-		(*JsonObject)->TryGetStringField(GET_MEMBER_NAME_STRING_CHECKED(UPropertyAnimatorVectorContext, AmplitudeMax), JsonAmplitudeMax);
+		FString AmplitudeMaxValue = AmplitudeMax.ToString();
+		ContextArchive->Get(GET_MEMBER_NAME_STRING_CHECKED(UPropertyAnimatorVectorContext, AmplitudeMax), AmplitudeMaxValue);
 		FVector ParseAmplitudeMax;
-		ParseAmplitudeMax.InitFromString(JsonAmplitudeMax);
+		ParseAmplitudeMax.InitFromString(AmplitudeMaxValue);
 		SetAmplitudeMax(ParseAmplitudeMax);
 
 		return true;
@@ -113,14 +111,14 @@ bool UPropertyAnimatorVectorContext::ImportPreset(const UPropertyAnimatorCorePre
 	return false;
 }
 
-bool UPropertyAnimatorVectorContext::ExportPreset(const UPropertyAnimatorCorePresetBase* InPreset, TSharedPtr<FJsonValue>& OutValue)
+bool UPropertyAnimatorVectorContext::ExportPreset(const UPropertyAnimatorCorePresetBase* InPreset, TSharedPtr<FPropertyAnimatorCorePresetArchive>& OutValue) const
 {
-	TSharedPtr<FJsonObject>* JsonObject = nullptr;
-
-	if (Super::ExportPreset(InPreset, OutValue) && OutValue->TryGetObject(JsonObject))
+	if (Super::ExportPreset(InPreset, OutValue) && OutValue->IsObject())
 	{
-		(*JsonObject)->SetStringField(GET_MEMBER_NAME_STRING_CHECKED(UPropertyAnimatorVectorContext, AmplitudeMin), AmplitudeMin.ToString());
-		(*JsonObject)->SetStringField(GET_MEMBER_NAME_STRING_CHECKED(UPropertyAnimatorVectorContext, AmplitudeMax), AmplitudeMax.ToString());
+		const TSharedPtr<FPropertyAnimatorCorePresetObjectArchive> ContextArchive = OutValue->AsMutableObject();
+
+		ContextArchive->Set(GET_MEMBER_NAME_STRING_CHECKED(UPropertyAnimatorVectorContext, AmplitudeMin), AmplitudeMin.ToString());
+		ContextArchive->Set(GET_MEMBER_NAME_STRING_CHECKED(UPropertyAnimatorVectorContext, AmplitudeMax), AmplitudeMax.ToString());
 
 		return true;
 	}
diff --git a/Engine/Plugins/Experimental/PropertyAnimator/Source/PropertyAnimator/Public/Animators/PropertyAnimatorClock.h b/Engine/Plugins/Experimental/PropertyAnimator/Source/PropertyAnimator/Public/Animators/PropertyAnimatorClock.h
--- a/Engine/Plugins/Experimental/PropertyAnimator/Source/PropertyAnimator/Public/Animators/PropertyAnimatorClock.h
+++ b/Engine/Plugins/Experimental/PropertyAnimator/Source/PropertyAnimator/Public/Animators/PropertyAnimatorClock.h
@@ -42,8 +42,8 @@ protected:
 	//~ Begin UPropertyAnimatorCoreBase
 	virtual void OnAnimatorRegistered(FPropertyAnimatorCoreMetadata& InMetadata) override;
 	virtual void EvaluateProperties(FInstancedPropertyBag& InParameters) override;
-	virtual bool ImportPreset(const UPropertyAnimatorCorePresetBase* InPreset, const TSharedRef<FJsonValue>& InValue) override;
-	virtual bool ExportPreset(const UPropertyAnimatorCorePresetBase* InPreset, TSharedPtr<FJsonValue>& OutValue) override;
+	virtual bool ImportPreset(const UPropertyAnimatorCorePresetBase* InPreset, const TSharedRef<FPropertyAnimatorCorePresetArchive>& InValue) override;
+	virtual bool ExportPreset(const UPropertyAnimatorCorePresetBase* InPreset, TSharedPtr<FPropertyAnimatorCorePresetArchive>& OutValue) const override;
 	//~ End UPropertyAnimatorCoreBase
 
 	/**
diff --git a/Engine/Plugins/Experimental/PropertyAnimator/Source/PropertyAnimator/Public/Animators/PropertyAnimatorCounter.h b/Engine/Plugins/Experimental/PropertyAnimator/Source/PropertyAnimator/Public/Animators/PropertyAnimatorCounter.h
--- a/Engine/Plugins/Experimental/PropertyAnimator/Source/PropertyAnimator/Public/Animators/PropertyAnimatorCounter.h
+++ b/Engine/Plugins/Experimental/PropertyAnimator/Source/PropertyAnimator/Public/Animators/PropertyAnimatorCounter.h
@@ -137,8 +137,8 @@ protected:
 	//~ Begin UPropertyAnimatorCoreBase
 	virtual void OnAnimatorRegistered(FPropertyAnimatorCoreMetadata& InMetadata) override;
 	virtual void EvaluateProperties(FInstancedPropertyBag& InParameters) override;
-	virtual bool ImportPreset(const UPropertyAnimatorCorePresetBase* InPreset, const TSharedRef<FJsonValue>& InValue) override;
-	virtual bool ExportPreset(const UPropertyAnimatorCorePresetBase* InPreset, TSharedPtr<FJsonValue>& OutValue) override;
+	virtual bool ImportPreset(const UPropertyAnimatorCorePresetBase* InPreset, const TSharedRef<FPropertyAnimatorCorePresetArchive>& InValue) override;
+	virtual bool ExportPreset(const UPropertyAnimatorCorePresetBase* InPreset, TSharedPtr<FPropertyAnimatorCorePresetArchive>& OutValue) const override;
 	//~ End UPropertyAnimatorCoreBase
 
 	UFUNCTION()
diff --git a/Engine/Plugins/Experimental/PropertyAnimator/Source/PropertyAnimator/Public/Animators/PropertyAnimatorCurve.h b/Engine/Plugins/Experimental/PropertyAnimator/Source/PropertyAnimator/Public/Animators/PropertyAnimatorCurve.h
--- a/Engine/Plugins/Experimental/PropertyAnimator/Source/PropertyAnimator/Public/Animators/PropertyAnimatorCurve.h
+++ b/Engine/Plugins/Experimental/PropertyAnimator/Source/PropertyAnimator/Public/Animators/PropertyAnimatorCurve.h
@@ -76,8 +76,8 @@ protected:
 	//~ Begin UPropertyAnimatorFloatBase
 	virtual void OnAnimatorRegistered(FPropertyAnimatorCoreMetadata& InMetadata) override;
 	virtual bool EvaluateProperty(const FPropertyAnimatorCoreData& InPropertyData, UPropertyAnimatorCoreContext* InContext, FInstancedPropertyBag& InParameters, FInstancedPropertyBag& OutEvaluationResult) const override;
-	virtual bool ImportPreset(const UPropertyAnimatorCorePresetBase* InPreset, const TSharedRef<FJsonValue>& InValue) override;
-	virtual bool ExportPreset(const UPropertyAnimatorCorePresetBase* InPreset, TSharedPtr<FJsonValue>& OutValue) override;
+	virtual bool ImportPreset(const UPropertyAnimatorCorePresetBase* InPreset, const TSharedRef<FPropertyAnimatorCorePresetArchive>& InValue) override;
+	virtual bool ExportPreset(const UPropertyAnimatorCorePresetBase* InPreset, TSharedPtr<FPropertyAnimatorCorePresetArchive>& OutValue) const override;
 	//~ End UPropertyAnimatorFloatBase
 
 	void OnEaseInChanged();
diff --git a/Engine/Plugins/Experimental/PropertyAnimator/Source/PropertyAnimator/Public/Animators/PropertyAnimatorNumericBase.h b/Engine/Plugins/Experimental/PropertyAnimator/Source/PropertyAnimator/Public/Animators/PropertyAnimatorNumericBase.h
--- a/Engine/Plugins/Experimental/PropertyAnimator/Source/PropertyAnimator/Public/Animators/PropertyAnimatorNumericBase.h
+++ b/Engine/Plugins/Experimental/PropertyAnimator/Source/PropertyAnimator/Public/Animators/PropertyAnimatorNumericBase.h
@@ -90,8 +90,8 @@ protected:
 	virtual void OnPropertyLinked(UPropertyAnimatorCoreContext* InLinkedProperty, EPropertyAnimatorPropertySupport InSupport) override;
 	virtual bool IsTimeSourceSupported(UPropertyAnimatorCoreTimeSourceBase* InTimeSource) const override;
 	virtual void OnAnimatorRegistered(FPropertyAnimatorCoreMetadata& InMetadata) override;
-	virtual bool ImportPreset(const UPropertyAnimatorCorePresetBase* InPreset, const TSharedRef<FJsonValue>& InValue) override;
-	virtual bool ExportPreset(const UPropertyAnimatorCorePresetBase* InPreset, TSharedPtr<FJsonValue>& OutValue) override;
+	virtual bool ImportPreset(const UPropertyAnimatorCorePresetBase* InPreset, const TSharedRef<FPropertyAnimatorCorePresetArchive>& InValue) override;
+	virtual bool ExportPreset(const UPropertyAnimatorCorePresetBase* InPreset, TSharedPtr<FPropertyAnimatorCorePresetArchive>& OutValue) const override;
 	//~ End UPropertyAnimatorCoreBase
 
 	/** Evaluate and return float value for a property */
diff --git a/Engine/Plugins/Experimental/PropertyAnimator/Source/PropertyAnimator/Public/Animators/PropertyAnimatorSoundWave.h b/Engine/Plugins/Experimental/PropertyAnimator/Source/PropertyAnimator/Public/Animators/PropertyAnimatorSoundWave.h
--- a/Engine/Plugins/Experimental/PropertyAnimator/Source/PropertyAnimator/Public/Animators/PropertyAnimatorSoundWave.h
+++ b/Engine/Plugins/Experimental/PropertyAnimator/Source/PropertyAnimator/Public/Animators/PropertyAnimatorSoundWave.h
@@ -43,8 +43,8 @@ protected:
 	//~ Begin UPropertyAnimatorFloatBase
 	virtual void OnAnimatorRegistered(FPropertyAnimatorCoreMetadata& InMetadata) override;
 	virtual bool EvaluateProperty(const FPropertyAnimatorCoreData& InPropertyData, UPropertyAnimatorCoreContext* InContext, FInstancedPropertyBag& InParameters, FInstancedPropertyBag& OutEvaluationResult) const override;
-	virtual bool ImportPreset(const UPropertyAnimatorCorePresetBase* InPreset, const TSharedRef<FJsonValue>& InValue) override;
-	virtual bool ExportPreset(const UPropertyAnimatorCorePresetBase* InPreset, TSharedPtr<FJsonValue>& OutValue) override;
+	virtual bool ImportPreset(const UPropertyAnimatorCorePresetBase* InPreset, const TSharedRef<FPropertyAnimatorCorePresetArchive>& InValue) override;
+	virtual bool ExportPreset(const UPropertyAnimatorCorePresetBase* InPreset, TSharedPtr<FPropertyAnimatorCorePresetArchive>& OutValue) const override;
 	//~ End UPropertyAnimatorFloatBase
 
 	/**
diff --git a/Engine/Plugins/Experimental/PropertyAnimator/Source/PropertyAnimator/Public/Properties/PropertyAnimatorFloatContext.h b/Engine/Plugins/Experimental/PropertyAnimator/Source/PropertyAnimator/Public/Properties/PropertyAnimatorFloatContext.h
--- a/Engine/Plugins/Experimental/PropertyAnimator/Source/PropertyAnimator/Public/Properties/PropertyAnimatorFloatContext.h
+++ b/Engine/Plugins/Experimental/PropertyAnimator/Source/PropertyAnimator/Public/Properties/PropertyAnimatorFloatContext.h
@@ -34,8 +34,8 @@ protected:
 	//~ Begin UPropertyAnimatorCoreContext
 	virtual bool EvaluateProperty(const FPropertyAnimatorCoreData& InProperty, const FInstancedPropertyBag& InAnimatorResult, FInstancedPropertyBag& OutEvaluatedValues) override;
 	virtual void OnAnimatedPropertyLinked() override;
-	virtual bool ImportPreset(const UPropertyAnimatorCorePresetBase* InPreset, const TSharedRef<FJsonValue>& InValue) override;
-	virtual bool ExportPreset(const UPropertyAnimatorCorePresetBase* InPreset, TSharedPtr<FJsonValue>& OutValue) override;
+	virtual bool ImportPreset(const UPropertyAnimatorCorePresetBase* InPreset, const TSharedRef<FPropertyAnimatorCorePresetArchive>& InValue) override;
+	virtual bool ExportPreset(const UPropertyAnimatorCorePresetBase* InPreset, TSharedPtr<FPropertyAnimatorCorePresetArchive>& OutValue) const override;
 	//~ End UPropertyAnimatorCoreContext
 
 	double GetClampedAmplitude(double InAmplitude);
diff --git a/Engine/Plugins/Experimental/PropertyAnimator/Source/PropertyAnimator/Public/Properties/PropertyAnimatorRotatorContext.h b/Engine/Plugins/Experimental/PropertyAnimator/Source/PropertyAnimator/Public/Properties/PropertyAnimatorRotatorContext.h
--- a/Engine/Plugins/Experimental/PropertyAnimator/Source/PropertyAnimator/Public/Properties/PropertyAnimatorRotatorContext.h
+++ b/Engine/Plugins/Experimental/PropertyAnimator/Source/PropertyAnimator/Public/Properties/PropertyAnimatorRotatorContext.h
@@ -34,8 +34,8 @@ protected:
 	//~ Begin UPropertyAnimatorCoreContext
 	virtual bool EvaluateProperty(const FPropertyAnimatorCoreData& InProperty, const FInstancedPropertyBag& InAnimatorResult, FInstancedPropertyBag& OutEvaluatedValues) override;
 	virtual void OnAnimatedPropertyLinked() override;
-	virtual bool ImportPreset(const UPropertyAnimatorCorePresetBase* InPreset, const TSharedRef<FJsonValue>& InValue) override;
-	virtual bool ExportPreset(const UPropertyAnimatorCorePresetBase* InPreset, TSharedPtr<FJsonValue>& OutValue) override;
+	virtual bool ImportPreset(const UPropertyAnimatorCorePresetBase* InPreset, const TSharedRef<FPropertyAnimatorCorePresetArchive>& InValue) override;
+	virtual bool ExportPreset(const UPropertyAnimatorCorePresetBase* InPreset, TSharedPtr<FPropertyAnimatorCorePresetArchive>& OutValue) const override;
 	//~ End UPropertyAnimatorCoreContext
 
 	FRotator GetClampedAmplitude(FRotator InAmplitude);
diff --git a/Engine/Plugins/Experimental/PropertyAnimator/Source/PropertyAnimator/Public/Properties/PropertyAnimatorVectorContext.h b/Engine/Plugins/Experimental/PropertyAnimator/Source/PropertyAnimator/Public/Properties/PropertyAnimatorVectorContext.h
--- a/Engine/Plugins/Experimental/PropertyAnimator/Source/PropertyAnimator/Public/Properties/PropertyAnimatorVectorContext.h
+++ b/Engine/Plugins/Experimental/PropertyAnimator/Source/PropertyAnimator/Public/Properties/PropertyAnimatorVectorContext.h
@@ -34,8 +34,8 @@ protected:
 	//~ Begin UPropertyAnimatorCoreContext
 	virtual bool EvaluateProperty(const FPropertyAnimatorCoreData& InProperty, const FInstancedPropertyBag& InAnimatorResult, FInstancedPropertyBag& OutEvaluatedValues) override;
 	virtual void OnAnimatedPropertyLinked() override;
-	virtual bool ImportPreset(const UPropertyAnimatorCorePresetBase* InPreset, const TSharedRef<FJsonValue>& InValue) override;
-	virtual bool ExportPreset(const UPropertyAnimatorCorePresetBase* InPreset, TSharedPtr<FJsonValue>& OutValue) override;
+	virtual bool ImportPreset(const UPropertyAnimatorCorePresetBase* InPreset, const TSharedRef<FPropertyAnimatorCorePresetArchive>& InValue) override;
+	virtual bool ExportPreset(const UPropertyAnimatorCorePresetBase* InPreset, TSharedPtr<FPropertyAnimatorCorePresetArchive>& OutValue) const override;
 	//~ End UPropertyAnimatorCoreContext
 
 	FVector GetClampedAmplitude(FVector InAmplitude);
diff --git a/Engine/Plugins/Experimental/PropertyAnimatorCore/Source/PropertyAnimatorCore/Private/Animators/PropertyAnimatorCoreBase.cpp b/Engine/Plugins/Experimental/PropertyAnimatorCore/Source/PropertyAnimatorCore/Private/Animators/PropertyAnimatorCoreBase.cpp
--- a/Engine/Plugins/Experimental/PropertyAnimatorCore/Source/PropertyAnimatorCore/Private/Animators/PropertyAnimatorCoreBase.cpp
+++ b/Engine/Plugins/Experimental/PropertyAnimatorCore/Source/PropertyAnimatorCore/Private/Animators/PropertyAnimatorCoreBase.cpp
@@ -3,7 +3,6 @@
 #include "Animators/PropertyAnimatorCoreBase.h"
 
 #include "Components/PropertyAnimatorCoreComponent.h"
-#include "Dom/JsonObject.h"
 #include "GameFramework/Actor.h"
 #include "Properties/PropertyAnimatorCoreResolver.h"
 #include "Properties/Handlers/PropertyAnimatorCoreHandlerBase.h"
@@ -207,102 +206,110 @@ void UPropertyAnimatorCoreBase::PostEditChangeProperty(FPropertyChangedEvent& Pr
 }
 #endif
 
-bool UPropertyAnimatorCoreBase::ImportPreset(const UPropertyAnimatorCorePresetBase* InPreset, const TSharedRef<FJsonValue>& InValue)
+bool UPropertyAnimatorCoreBase::ImportPreset(const UPropertyAnimatorCorePresetBase* InPreset, const TSharedRef<FPropertyAnimatorCorePresetArchive>& InValue)
 {
-	const TSharedPtr<FJsonObject>* JsonAnimatorObject;
-	if (!InValue->TryGetObject(JsonAnimatorObject) || !JsonAnimatorObject)
+	if (!InValue->IsObject())
 	{
 		return false;
 	}
 
-	bool bJsonEnabled = bAnimatorEnabled;
-	(*JsonAnimatorObject)->TryGetBoolField(GET_MEMBER_NAME_STRING_CHECKED(UPropertyAnimatorCoreBase, bAnimatorEnabled), bJsonEnabled);
-	SetAnimatorEnabled(bJsonEnabled);
+	const TSharedPtr<FPropertyAnimatorCorePresetObjectArchive> AnimatorArchive = InValue->AsMutableObject();
 
-	FString JsonDisplayName = AnimatorDisplayName.ToString();
-	(*JsonAnimatorObject)->TryGetStringField(GET_MEMBER_NAME_STRING_CHECKED(UPropertyAnimatorCoreBase, AnimatorDisplayName), JsonDisplayName);
-	SetAnimatorDisplayName(FName(JsonDisplayName));
+	bool bEnabledValue = bAnimatorEnabled;
+	AnimatorArchive->Get(GET_MEMBER_NAME_STRING_CHECKED(UPropertyAnimatorCoreBase, bAnimatorEnabled), bEnabledValue);
+	SetAnimatorEnabled(bEnabledValue);
 
-	const TArray<TSharedPtr<FJsonValue>>* JsonValues;
-	(*JsonAnimatorObject)->TryGetArrayField(GET_MEMBER_NAME_STRING_CHECKED(UPropertyAnimatorCoreBase, LinkedProperties), JsonValues);
-	for (const TSharedPtr<FJsonValue>& JsonValue : (*JsonValues))
+	FString DisplayNameValue = AnimatorDisplayName.ToString();
+	AnimatorArchive->Get(GET_MEMBER_NAME_STRING_CHECKED(UPropertyAnimatorCoreBase, AnimatorDisplayName), DisplayNameValue);
+	SetAnimatorDisplayName(FName(DisplayNameValue));
+
+	TSharedPtr<FPropertyAnimatorCorePresetArchive> LinkedPropertiesArchive;
+	AnimatorArchive->Get(GET_MEMBER_NAME_STRING_CHECKED(UPropertyAnimatorCoreBase, LinkedProperties), LinkedPropertiesArchive);
+	if (const TSharedPtr<FPropertyAnimatorCorePresetArrayArchive> LinkedPropertiesArray = LinkedPropertiesArchive->AsMutableArray())
 	{
-		TSharedPtr<FJsonObject>* JsonPropertyObject;
-		if (!JsonValue->TryGetObject(JsonPropertyObject) || !JsonPropertyObject || !JsonPropertyObject->IsValid())
+		for (int32 Index = 0; Index < LinkedPropertiesArray->Num(); Index++)
 		{
-			continue;
-		}
+			TSharedPtr<FPropertyAnimatorCorePresetArchive> LinkedPropertyArchive;
+			if (!LinkedPropertiesArray->Get(Index, LinkedPropertyArchive) || !LinkedPropertyArchive->IsObject())
+			{
+				continue;
+			}
 
-		FString AnimatedPropertyLocatorPath;
-		if (!(*JsonPropertyObject)->TryGetStringField(GET_MEMBER_NAME_STRING_CHECKED(UPropertyAnimatorCoreContext, AnimatedProperty), AnimatedPropertyLocatorPath))
-		{
-			continue;
-		}
+			const TSharedPtr<FPropertyAnimatorCorePresetObjectArchive> LinkedPropertyObject = LinkedPropertyArchive->AsMutableObject();
+
+			FString AnimatedPropertyLocatorPath;
+			if (!LinkedPropertyObject->Get(GET_MEMBER_NAME_STRING_CHECKED(UPropertyAnimatorCoreContext, AnimatedProperty), AnimatedPropertyLocatorPath))
+			{
+				continue;
+			}
 
-		FPropertyAnimatorCoreData PropertyData(GetAnimatorActor(), AnimatedPropertyLocatorPath);
+			FPropertyAnimatorCoreData PropertyData(GetAnimatorActor(), AnimatedPropertyLocatorPath);
 
-		if (!PropertyData.IsResolved())
-		{
-			continue;
-		}
+			if (!PropertyData.IsResolved())
+			{
+				continue;
+			}
 
-		if (IPropertyAnimatorCorePresetable* PropertyContext = Cast<IPropertyAnimatorCorePresetable>(LinkProperty(PropertyData)))
-		{
-			PropertyContext->ImportPreset(InPreset, JsonValue.ToSharedRef());
+			if (IPropertyAnimatorCorePresetable* PropertyContext = Cast<IPropertyAnimatorCorePresetable>(LinkProperty(PropertyData)))
+			{
+				PropertyContext->ImportPreset(InPreset, LinkedPropertyArchive.ToSharedRef());
+			}
 		}
 	}
 
-	bool bJsonOverrideTimeSource = bOverrideTimeSource;
-	(*JsonAnimatorObject)->TryGetBoolField(GET_MEMBER_NAME_STRING_CHECKED(UPropertyAnimatorCoreBase, bOverrideTimeSource), bJsonOverrideTimeSource);
-	SetOverrideTimeSource(bJsonOverrideTimeSource);
+	bool bOverrideTimeSourceValue = bOverrideTimeSource;
+	AnimatorArchive->Get(GET_MEMBER_NAME_STRING_CHECKED(UPropertyAnimatorCoreBase, bOverrideTimeSource), bOverrideTimeSourceValue);
+	SetOverrideTimeSource(bOverrideTimeSourceValue);
 
-	FString JsonTimeSourceName = TimeSourceName.ToString();
-	(*JsonAnimatorObject)->TryGetStringField(GET_MEMBER_NAME_STRING_CHECKED(UPropertyAnimatorCoreBase, TimeSourceName), JsonTimeSourceName);
-	SetTimeSourceName(FName(JsonTimeSourceName));
+	FString TimeSourceNameValue = TimeSourceName.ToString();
+	AnimatorArchive->Get(GET_MEMBER_NAME_STRING_CHECKED(UPropertyAnimatorCoreBase, TimeSourceName), TimeSourceNameValue);
+	SetTimeSourceName(FName(TimeSourceNameValue));
 
 	if (UPropertyAnimatorCoreTimeSourceBase* TimeSource = FindOrAddTimeSource(GetTimeSourceName()))
 	{
-		if ((*JsonAnimatorObject)->HasTypedField(GET_MEMBER_NAME_STRING_CHECKED(UPropertyAnimatorCoreBase, ActiveTimeSource), EJson::Object))
+		if (AnimatorArchive->Has(GET_MEMBER_NAME_STRING_CHECKED(UPropertyAnimatorCoreBase, ActiveTimeSource), EPropertyAnimatorCorePresetArchiveType::Object))
 		{
-			TimeSource->ImportPreset(InPreset, (*JsonAnimatorObject)->TryGetField(GET_MEMBER_NAME_STRING_CHECKED(UPropertyAnimatorCoreBase, ActiveTimeSource)).ToSharedRef());
+			TSharedPtr<FPropertyAnimatorCorePresetArchive> TimeSourceArchive;
+			AnimatorArchive->Get(GET_MEMBER_NAME_STRING_CHECKED(UPropertyAnimatorCoreBase, ActiveTimeSource), TimeSourceArchive);
+			TimeSource->ImportPreset(InPreset, TimeSourceArchive.ToSharedRef());
 		}
 	}
 
 	return true;
 }
 
-bool UPropertyAnimatorCoreBase::ExportPreset(const UPropertyAnimatorCorePresetBase* InPreset, TSharedPtr<FJsonValue>& OutValue)
+bool UPropertyAnimatorCoreBase::ExportPreset(const UPropertyAnimatorCorePresetBase* InPreset, TSharedPtr<FPropertyAnimatorCorePresetArchive>& OutValue) const
 {
-	TSharedRef<FJsonObject> JsonAnimatorObject = MakeShared<FJsonObject>();
-	OutValue = MakeShared<FJsonValueObject>(JsonAnimatorObject);
+	TSharedPtr<FPropertyAnimatorCorePresetObjectArchive> AnimatorArchive = InPreset->GetArchiveImplementation()->CreateObject();
+	OutValue = AnimatorArchive;
 
-	JsonAnimatorObject->SetStringField(TEXT("AnimatorClass"), GetClass()->GetClassPathName().ToString());
-	JsonAnimatorObject->SetBoolField(GET_MEMBER_NAME_STRING_CHECKED(UPropertyAnimatorCoreBase, bAnimatorEnabled), bAnimatorEnabled);
-	JsonAnimatorObject->SetStringField(GET_MEMBER_NAME_STRING_CHECKED(UPropertyAnimatorCoreBase, AnimatorDisplayName), AnimatorDisplayName.ToString());
+	AnimatorArchive->Set(TEXT("AnimatorClass"), GetClass()->GetClassPathName().ToString());
+	AnimatorArchive->Set(GET_MEMBER_NAME_STRING_CHECKED(UPropertyAnimatorCoreBase, bAnimatorEnabled), bAnimatorEnabled);
+	AnimatorArchive->Set(GET_MEMBER_NAME_STRING_CHECKED(UPropertyAnimatorCoreBase, AnimatorDisplayName), AnimatorDisplayName.ToString());
 
-	TArray<TSharedPtr<FJsonValue>> JsonValues;
+	TSharedPtr<FPropertyAnimatorCorePresetArrayArchive> LinkedPropertiesArchive = InPreset->GetArchiveImplementation()->CreateArray();
 	for (const TObjectPtr<UPropertyAnimatorCoreContext>& LinkedProperty : LinkedProperties)
 	{
 		if (IPropertyAnimatorCorePresetable* PropertyContext = Cast<IPropertyAnimatorCorePresetable>(LinkedProperty))
 		{
-			TSharedPtr<FJsonValue> JsonValue;
-			if (PropertyContext->ExportPreset(InPreset, JsonValue) && JsonValue.IsValid())
+			TSharedPtr<FPropertyAnimatorCorePresetArchive> LinkedPropertyArchive;
+			if (PropertyContext->ExportPreset(InPreset, LinkedPropertyArchive) && LinkedPropertyArchive.IsValid())
 			{
-				JsonValues.Add(JsonValue);
+				LinkedPropertiesArchive->Add(LinkedPropertyArchive.ToSharedRef());
 			}
 		}
 	}
-	JsonAnimatorObject->SetArrayField(GET_MEMBER_NAME_STRING_CHECKED(UPropertyAnimatorCoreBase, LinkedProperties), JsonValues);
+	AnimatorArchive->Set(GET_MEMBER_NAME_STRING_CHECKED(UPropertyAnimatorCoreBase, LinkedProperties), LinkedPropertiesArchive.ToSharedRef());
 
-	JsonAnimatorObject->SetBoolField(GET_MEMBER_NAME_STRING_CHECKED(UPropertyAnimatorCoreBase, bOverrideTimeSource), bOverrideTimeSource);
-	JsonAnimatorObject->SetStringField(GET_MEMBER_NAME_STRING_CHECKED(UPropertyAnimatorCoreBase, TimeSourceName), TimeSourceName.ToString());
+	AnimatorArchive->Set(GET_MEMBER_NAME_STRING_CHECKED(UPropertyAnimatorCoreBase, bOverrideTimeSource), bOverrideTimeSource);
+	AnimatorArchive->Set(GET_MEMBER_NAME_STRING_CHECKED(UPropertyAnimatorCoreBase, TimeSourceName), TimeSourceName.ToString());
 
 	if (UPropertyAnimatorCoreTimeSourceBase* TimeSource = GetActiveTimeSource())
 	{
-		TSharedPtr<FJsonValue> JsonValue;
-		if (TimeSource->ExportPreset(InPreset, JsonValue) && JsonValue.IsValid())
+		TSharedPtr<FPropertyAnimatorCorePresetArchive> TimeSourceArchive;
+		if (TimeSource->ExportPreset(InPreset, TimeSourceArchive) && TimeSourceArchive.IsValid())
 		{
-			JsonAnimatorObject->SetField(GET_MEMBER_NAME_STRING_CHECKED(UPropertyAnimatorCoreBase, ActiveTimeSource), JsonValue);
+			AnimatorArchive->Set(GET_MEMBER_NAME_STRING_CHECKED(UPropertyAnimatorCoreBase, ActiveTimeSource), TimeSourceArchive.ToSharedRef());
 		}
 	}
 
diff --git a/Engine/Plugins/Experimental/PropertyAnimatorCore/Source/PropertyAnimatorCore/Private/Presets/PropertyAnimatorCoreAnimatorPreset.cpp b/Engine/Plugins/Experimental/PropertyAnimatorCore/Source/PropertyAnimatorCore/Private/Presets/PropertyAnimatorCoreAnimatorPreset.cpp
--- a/Engine/Plugins/Experimental/PropertyAnimatorCore/Source/PropertyAnimatorCore/Private/Presets/PropertyAnimatorCoreAnimatorPreset.cpp
+++ b/Engine/Plugins/Experimental/PropertyAnimatorCore/Source/PropertyAnimatorCore/Private/Presets/PropertyAnimatorCoreAnimatorPreset.cpp
@@ -3,12 +3,7 @@
 #include "Presets/PropertyAnimatorCoreAnimatorPreset.h"
 
 #include "Animators/PropertyAnimatorCoreBase.h"
-#include "Dom/JsonObject.h"
-#include "Dom/JsonValue.h"
 #include "Presets/PropertyAnimatorCorePresetable.h"
-#include "Serialization/JsonReader.h"
-#include "Serialization/JsonSerializer.h"
-#include "Serialization/JsonWriter.h"
 
 bool UPropertyAnimatorCoreAnimatorPreset::IsPresetApplied(const UPropertyAnimatorCoreBase* InAnimator) const
 {
@@ -22,7 +17,7 @@ bool UPropertyAnimatorCoreAnimatorPreset::IsPresetSupported(const AActor* InActo
 
 bool UPropertyAnimatorCoreAnimatorPreset::ApplyPreset(UPropertyAnimatorCoreBase* InAnimator)
 {
-	if (InAnimator->IsTemplate())
+	if (InAnimator->IsTemplate() || !AnimatorPreset.IsValid())
 	{
 		return false;
 	}
@@ -39,22 +34,18 @@ void UPropertyAnimatorCoreAnimatorPreset::CreatePreset(FName InName, const TArra
 {
 	Super::CreatePreset(InName, InPresetableItem);
 
-	TSharedPtr<FJsonValue> JsonValue;
+	TSharedPtr<FPropertyAnimatorCorePresetArchive> ItemArchive;
 
 	if (InPresetableItem[0]
-		&& InPresetableItem[0]->ExportPreset(this, JsonValue)
-		&& JsonValue.IsValid())
+		&& InPresetableItem[0]->ExportPreset(this, ItemArchive)
+		&& ItemArchive.IsValid())
 	{
-		TSharedPtr<FJsonObject>* JsonObject = nullptr;
-		if (JsonValue->TryGetObject(JsonObject))
+		FString OutputString;
+		if (ItemArchive->ToString(OutputString))
 		{
-			FString JsonString;
-			TSharedRef<TJsonWriter<>> JsonWriter = TJsonWriterFactory<>::Create(&JsonString);
-			if (FJsonSerializer::Serialize(JsonObject->ToSharedRef(), JsonWriter))
-			{
-				PresetVersion = 0;
-				PresetContent = JsonString;
-			}
+			PresetVersion = 0;
+			PresetFormat = ItemArchive->GetImplementationType();
+			PresetContent = OutputString;
 		}
 	}
 }
@@ -66,21 +57,19 @@ bool UPropertyAnimatorCoreAnimatorPreset::LoadPreset()
 		return false;
 	}
 
-	TSharedRef<TJsonReader<>> Reader = TJsonReaderFactory<>::Create(PresetContent);
-
-	TSharedPtr<FJsonObject> JsonObject;
-	if (!FJsonSerializer::Deserialize(Reader, JsonObject) || !JsonObject.IsValid())
+	TSharedRef<FPropertyAnimatorCorePresetObjectArchive> ObjectArchive = GetArchiveImplementation()->CreateObject();
+	if (!ObjectArchive->FromString(PresetContent))
 	{
 		return false;
 	}
 
 	FString Class;
-	JsonObject->TryGetStringField(TEXT("AnimatorClass"), Class);
+	ObjectArchive->Get(TEXT("AnimatorClass"), Class);
 
 	if (UClass* AnimatorClass = LoadObject<UClass>(nullptr, *Class))
 	{
 		TargetAnimatorClass = AnimatorClass;
-		AnimatorPreset = MakeShared<FJsonValueObject>(JsonObject);
+		AnimatorPreset = ObjectArchive;
 		return true;
 	}
 
diff --git a/Engine/Plugins/Experimental/PropertyAnimatorCore/Source/PropertyAnimatorCore/Private/Presets/PropertyAnimatorCoreJsonPresetArchive.cpp b/Engine/Plugins/Experimental/PropertyAnimatorCore/Source/PropertyAnimatorCore/Private/Presets/PropertyAnimatorCoreJsonPresetArchive.cpp
new file mode 100644
--- /dev/null
+++ b/Engine/Plugins/Experimental/PropertyAnimatorCore/Source/PropertyAnimatorCore/Private/Presets/PropertyAnimatorCoreJsonPresetArchive.cpp
@@ -0,0 +1,404 @@
+// Copyright Epic Games, Inc. All Rights Reserved.
+
+#include "Presets/PropertyAnimatorCoreJsonPresetArchive.h"
+
+#include "Dom/JsonObject.h"
+#include "Dom/JsonValue.h"
+#include "Serialization/JsonReader.h"
+#include "Serialization/JsonSerializer.h"
+#include "Serialization/JsonWriter.h"
+
+TSharedPtr<FPropertyAnimatorCorePresetJsonArchiveImplementation> FPropertyAnimatorCorePresetJsonArchiveImplementation::Instance;
+
+TSharedPtr<FPropertyAnimatorCorePresetJsonArchiveImplementation> FPropertyAnimatorCorePresetJsonArchiveImplementation::Get()
+{
+	if (!Instance.IsValid())
+	{
+		Instance = MakeShared<FPropertyAnimatorCorePresetJsonArchiveImplementation>();
+	}
+
+	return Instance;
+}
+
+TSharedRef<FPropertyAnimatorCorePresetObjectArchive> FPropertyAnimatorCorePresetJsonArchiveImplementation::CreateObject()
+{
+	return MakeShared<FPropertyAnimatorCorePresetJsonObjectArchive>();
+}
+
+TSharedRef<FPropertyAnimatorCorePresetArrayArchive> FPropertyAnimatorCorePresetJsonArchiveImplementation::CreateArray()
+{
+	return MakeShared<FPropertyAnimatorCorePresetJsonArrayArchive>();
+}
+
+TSharedRef<FPropertyAnimatorCorePresetValueArchive> FPropertyAnimatorCorePresetJsonArchiveImplementation::CreateValue(bool bInValue)
+{
+	return MakeShared<FPropertyAnimatorCorePresetJsonValueArchive>(bInValue);
+}
+
+TSharedRef<FPropertyAnimatorCorePresetValueArchive> FPropertyAnimatorCorePresetJsonArchiveImplementation::CreateValue(const FString& InValue)
+{
+	return MakeShared<FPropertyAnimatorCorePresetJsonValueArchive>(InValue);
+}
+
+TSharedRef<FPropertyAnimatorCorePresetValueArchive> FPropertyAnimatorCorePresetJsonArchiveImplementation::CreateValue(uint64 InValue)
+{
+	return MakeShared<FPropertyAnimatorCorePresetJsonValueArchive>(InValue);
+}
+
+TSharedRef<FPropertyAnimatorCorePresetValueArchive> FPropertyAnimatorCorePresetJsonArchiveImplementation::CreateValue(int64 InValue)
+{
+	return MakeShared<FPropertyAnimatorCorePresetJsonValueArchive>(InValue);
+}
+
+TSharedRef<FPropertyAnimatorCorePresetValueArchive> FPropertyAnimatorCorePresetJsonArchiveImplementation::CreateValue(double InValue)
+{
+	return MakeShared<FPropertyAnimatorCorePresetJsonValueArchive>(InValue);
+}
+
+FPropertyAnimatorCorePresetJsonObjectArchive::FPropertyAnimatorCorePresetJsonObjectArchive()
+{
+	JsonObject = MakeShared<FJsonObject>();
+}
+
+FPropertyAnimatorCorePresetJsonObjectArchive::FPropertyAnimatorCorePresetJsonObjectArchive(TSharedRef<FJsonObject> InJsonObject)
+{
+	JsonObject = InJsonObject;
+}
+
+bool FPropertyAnimatorCorePresetJsonObjectArchive::Remove(const FString& InKey)
+{
+	if (!JsonObject->HasField(InKey))
+	{
+		return false;
+	}
+
+	JsonObject->RemoveField(InKey);
+	return true;
+}
+
+void FPropertyAnimatorCorePresetJsonObjectArchive::Clear()
+{
+	JsonObject->Values.Empty();
+}
+
+bool FPropertyAnimatorCorePresetJsonObjectArchive::Set(const FString& InKey, TSharedRef<FPropertyAnimatorCorePresetArchive> InValue)
+{
+	check(InValue->GetImplementationType().IsEqual(GetImplementationType()))
+
+	if (const FPropertyAnimatorCorePresetJsonArrayArchive* Array = static_cast<FPropertyAnimatorCorePresetJsonArrayArchive*>(InValue->AsMutableArray().Get()))
+	{
+		JsonObject->SetArrayField(InKey, Array->GetJsonValues());
+	}
+	else if (const FPropertyAnimatorCorePresetJsonObjectArchive* Object = static_cast<FPropertyAnimatorCorePresetJsonObjectArchive*>(InValue->AsMutableObject().Get()))
+	{
+		JsonObject->SetObjectField(InKey, Object->GetJsonObject());
+	}
+	else if (const FPropertyAnimatorCorePresetJsonValueArchive* Value = static_cast<FPropertyAnimatorCorePresetJsonValueArchive*>(InValue->AsMutableValue().Get()))
+	{
+		JsonObject->SetField(InKey, Value->GetJsonValue());
+	}
+
+	return true;
+}
+
+bool FPropertyAnimatorCorePresetJsonObjectArchive::Set(const FString& InKey, bool bInValue)
+{
+	return Set(InKey, GetImplementation()->CreateValue(bInValue));
+}
+
+bool FPropertyAnimatorCorePresetJsonObjectArchive::Set(const FString& InKey, uint64 InValue)
+{
+	return Set(InKey, GetImplementation()->CreateValue(InValue));
+}
+
+bool FPropertyAnimatorCorePresetJsonObjectArchive::Set(const FString& InKey, int64 InValue)
+{
+	return Set(InKey, GetImplementation()->CreateValue(InValue));
+}
+
+bool FPropertyAnimatorCorePresetJsonObjectArchive::Set(const FString& InKey, double InValue)
+{
+	return Set(InKey, GetImplementation()->CreateValue(InValue));
+}
+
+bool FPropertyAnimatorCorePresetJsonObjectArchive::Set(const FString& InKey, const FString& InValue)
+{
+	return Set(InKey, GetImplementation()->CreateValue(InValue));
+}
+
+bool FPropertyAnimatorCorePresetJsonObjectArchive::Get(const FString& InKey, TSharedPtr<FPropertyAnimatorCorePresetArchive>& OutValue) const
+{
+	TSharedPtr<FJsonValue> JsonValue = JsonObject->TryGetField(InKey);
+
+	if (!JsonValue.IsValid())
+	{
+		return false;
+	}
+
+	if (JsonValue->Type == EJson::Array)
+	{
+		OutValue = MakeShared<FPropertyAnimatorCorePresetJsonArrayArchive>(JsonValue->AsArray());
+	}
+	else if (JsonValue->Type == EJson::Object)
+	{
+		OutValue = MakeShared<FPropertyAnimatorCorePresetJsonObjectArchive>(JsonValue->AsObject().ToSharedRef());
+	}
+	else if (JsonValue->Type != EJson::None)
+	{
+		OutValue = MakeShared<FPropertyAnimatorCorePresetJsonValueArchive>(JsonValue.ToSharedRef());
+	}
+	else
+	{
+		return false;
+	}
+
+	return true;
+}
+
+bool FPropertyAnimatorCorePresetJsonObjectArchive::Get(const FString& InKey, bool& bOutValue) const
+{
+	TSharedPtr<FPropertyAnimatorCorePresetArchive> Value;
+
+	return Get(InKey, Value)
+		&& Value->IsValue()
+		&& Value->AsMutableValue()->Get(bOutValue);
+}
+
+bool FPropertyAnimatorCorePresetJsonObjectArchive::Get(const FString& InKey, uint64& OutValue) const
+{
+	TSharedPtr<FPropertyAnimatorCorePresetArchive> Value;
+
+	return Get(InKey, Value)
+		&& Value->IsValue()
+		&& Value->AsMutableValue()->Get(OutValue);
+}
+
+bool FPropertyAnimatorCorePresetJsonObjectArchive::Get(const FString& InKey, int64& OutValue) const
+{
+	TSharedPtr<FPropertyAnimatorCorePresetArchive> Value;
+
+	return Get(InKey, Value)
+		&& Value->IsValue()
+		&& Value->AsMutableValue()->Get(OutValue);
+}
+
+bool FPropertyAnimatorCorePresetJsonObjectArchive::Get(const FString& InKey, double& OutValue) const
+{
+	TSharedPtr<FPropertyAnimatorCorePresetArchive> Value;
+
+	return Get(InKey, Value)
+		&& Value->IsValue()
+		&& Value->AsMutableValue()->Get(OutValue);
+}
+
+bool FPropertyAnimatorCorePresetJsonObjectArchive::Get(const FString& InKey, FString& OutValue) const
+{
+	TSharedPtr<FPropertyAnimatorCorePresetArchive> Value;
+
+	return Get(InKey, Value)
+		&& Value->IsValue()
+		&& Value->AsMutableValue()->Get(OutValue);
+}
+
+bool FPropertyAnimatorCorePresetJsonObjectArchive::ToString(FString& OutString) const
+{
+	const TSharedRef<TJsonWriter<>> JsonWriter = TJsonWriterFactory<>::Create(&OutString);
+	return JsonObject.IsValid() && FJsonSerializer::Serialize(JsonObject.ToSharedRef(), JsonWriter);
+}
+
+bool FPropertyAnimatorCorePresetJsonObjectArchive::FromString(const FString& InString)
+{
+	const TSharedRef<TJsonReader<>> JsonReader = TJsonReaderFactory<>::Create(InString);
+	return FJsonSerializer::Deserialize(JsonReader, JsonObject) && JsonObject.IsValid();
+}
+
+TSharedRef<FPropertyAnimatorCorePresetArchiveImplementation> FPropertyAnimatorCorePresetJsonObjectArchive::GetImplementation() const
+{
+	return FPropertyAnimatorCorePresetJsonArchiveImplementation::Get().ToSharedRef();
+}
+
+FPropertyAnimatorCorePresetJsonArrayArchive::FPropertyAnimatorCorePresetJsonArrayArchive(const TArray<TSharedPtr<FJsonValue>>& InJsonArray)
+{
+	JsonValues = InJsonArray;
+}
+
+bool FPropertyAnimatorCorePresetJsonArrayArchive::Get(int32 InIndex, TSharedPtr<FPropertyAnimatorCorePresetArchive>& OutValue) const
+{
+	if (!JsonValues.IsValidIndex(InIndex) || !JsonValues[InIndex].IsValid())
+	{
+		return false;
+	}
+
+	if (JsonValues[InIndex]->Type == EJson::Array)
+	{
+		OutValue = MakeShared<FPropertyAnimatorCorePresetJsonArrayArchive>(JsonValues[InIndex]->AsArray());
+	}
+	else if (JsonValues[InIndex]->Type == EJson::Object)
+	{
+		OutValue = MakeShared<FPropertyAnimatorCorePresetJsonObjectArchive>(JsonValues[InIndex]->AsObject().ToSharedRef());
+	}
+	else if (JsonValues[InIndex]->Type != EJson::None)
+	{
+		OutValue = MakeShared<FPropertyAnimatorCorePresetJsonValueArchive>(JsonValues[InIndex].ToSharedRef());
+	}
+	else
+	{
+		return false;
+	}
+
+	return true;
+}
+
+int32 FPropertyAnimatorCorePresetJsonArrayArchive::Num() const
+{
+	return JsonValues.Num();
+}
+
+bool FPropertyAnimatorCorePresetJsonArrayArchive::Remove(int32 InIndex)
+{
+	if (!JsonValues.IsValidIndex(InIndex))
+	{
+		return false;
+	}
+
+	JsonValues.RemoveAt(InIndex);
+	return true;
+}
+
+void FPropertyAnimatorCorePresetJsonArrayArchive::Clear()
+{
+	JsonValues.Empty();
+}
+
+bool FPropertyAnimatorCorePresetJsonArrayArchive::Add(TSharedRef<FPropertyAnimatorCorePresetArchive> InValue)
+{
+	check(InValue->GetImplementationType().IsEqual(GetImplementationType()))
+
+	if (FPropertyAnimatorCorePresetJsonArrayArchive* Array = static_cast<FPropertyAnimatorCorePresetJsonArrayArchive*>(InValue->AsMutableArray().Get()))
+	{
+		JsonValues.Add(MakeShared<FJsonValueArray>(Array->GetJsonValues()));
+	}
+	else if (FPropertyAnimatorCorePresetJsonObjectArchive* Object = static_cast<FPropertyAnimatorCorePresetJsonObjectArchive*>(InValue->AsMutableObject().Get()))
+	{
+		JsonValues.Add(MakeShared<FJsonValueObject>(Object->GetJsonObject()));
+	}
+	else if (FPropertyAnimatorCorePresetJsonValueArchive* Value = static_cast<FPropertyAnimatorCorePresetJsonValueArchive*>(InValue->AsMutableValue().Get()))
+	{
+		JsonValues.Add(Value->GetJsonValue());
+	}
+
+	return true;
+}
+
+bool FPropertyAnimatorCorePresetJsonArrayArchive::Add(bool bInValue)
+{
+	return Add(GetImplementation()->CreateValue(bInValue));
+}
+
+bool FPropertyAnimatorCorePresetJsonArrayArchive::Add(uint64 InValue)
+{
+	return Add(GetImplementation()->CreateValue(InValue));
+}
+
+bool FPropertyAnimatorCorePresetJsonArrayArchive::Add(int64 InValue)
+{
+	return Add(GetImplementation()->CreateValue(InValue));
+}
+
+bool FPropertyAnimatorCorePresetJsonArrayArchive::Add(double InValue)
+{
+	return Add(GetImplementation()->CreateValue(InValue));
+}
+
+bool FPropertyAnimatorCorePresetJsonArrayArchive::Add(const FString& InValue)
+{
+	return Add(GetImplementation()->CreateValue(InValue));
+}
+
+bool FPropertyAnimatorCorePresetJsonArrayArchive::FromString(const FString& InString)
+{
+	const TSharedRef<TJsonReader<>> JsonReader = TJsonReaderFactory<>::Create(InString);
+	return FJsonSerializer::Deserialize(JsonReader, JsonValues);
+}
+
+bool FPropertyAnimatorCorePresetJsonArrayArchive::ToString(FString& OutString) const
+{
+	const TSharedRef<TJsonWriter<>> JsonWriter = TJsonWriterFactory<>::Create(&OutString);
+	return FJsonSerializer::Serialize(JsonValues, JsonWriter);
+}
+
+TSharedRef<FPropertyAnimatorCorePresetArchiveImplementation> FPropertyAnimatorCorePresetJsonArrayArchive::GetImplementation() const
+{
+	return FPropertyAnimatorCorePresetJsonArchiveImplementation::Get().ToSharedRef();
+}
+
+FPropertyAnimatorCorePresetJsonValueArchive::FPropertyAnimatorCorePresetJsonValueArchive(bool bInValue)
+{
+	JsonValue = MakeShared<FJsonValueBoolean>(bInValue);
+}
+
+FPropertyAnimatorCorePresetJsonValueArchive::FPropertyAnimatorCorePresetJsonValueArchive(uint64 InValue)
+{
+	JsonValue = MakeShared<FJsonValueNumber>(InValue);
+}
+
+FPropertyAnimatorCorePresetJsonValueArchive::FPropertyAnimatorCorePresetJsonValueArchive(int64 InValue)
+{
+	JsonValue = MakeShared<FJsonValueNumber>(InValue);
+}
+
+FPropertyAnimatorCorePresetJsonValueArchive::FPropertyAnimatorCorePresetJsonValueArchive(double InValue)
+{
+	JsonValue = MakeShared<FJsonValueNumber>(InValue);
+}
+
+FPropertyAnimatorCorePresetJsonValueArchive::FPropertyAnimatorCorePresetJsonValueArchive(const FString& InValue)
+{
+	JsonValue = MakeShared<FJsonValueString>(InValue);
+}
+
+FPropertyAnimatorCorePresetJsonValueArchive::FPropertyAnimatorCorePresetJsonValueArchive(TSharedRef<FJsonValue> InJsonValue)
+{
+	check(InJsonValue->Type != EJson::Array && InJsonValue->Type != EJson::Object && InJsonValue->Type != EJson::None)
+	JsonValue = InJsonValue;
+}
+
+bool FPropertyAnimatorCorePresetJsonValueArchive::FromString(const FString& InString)
+{
+	return false;
+}
+
+bool FPropertyAnimatorCorePresetJsonValueArchive::ToString(FString& OutString) const
+{
+	return false;
+}
+
+bool FPropertyAnimatorCorePresetJsonValueArchive::Get(double& OutValue) const
+{
+	return JsonValue.IsValid() && JsonValue->TryGetNumber(OutValue);
+}
+
+bool FPropertyAnimatorCorePresetJsonValueArchive::Get(bool& bOutValue) const
+{
+	return JsonValue.IsValid() && JsonValue->TryGetBool(bOutValue);
+}
+
+bool FPropertyAnimatorCorePresetJsonValueArchive::Get(FString& OutValue) const
+{
+	return JsonValue.IsValid() && JsonValue->TryGetString(OutValue);
+}
+
+bool FPropertyAnimatorCorePresetJsonValueArchive::Get(uint64& OutValue) const
+{
+	return JsonValue.IsValid() && JsonValue->TryGetNumber(OutValue);
+}
+
+bool FPropertyAnimatorCorePresetJsonValueArchive::Get(int64& OutValue) const
+{
+	return JsonValue.IsValid() && JsonValue->TryGetNumber(OutValue);
+}
+
+TSharedRef<FPropertyAnimatorCorePresetArchiveImplementation> FPropertyAnimatorCorePresetJsonValueArchive::GetImplementation() const
+{
+	return FPropertyAnimatorCorePresetJsonArchiveImplementation::Get().ToSharedRef();
+}
diff --git a/Engine/Plugins/Experimental/PropertyAnimatorCore/Source/PropertyAnimatorCore/Private/Presets/PropertyAnimatorCoreJsonPresetArchive.h b/Engine/Plugins/Experimental/PropertyAnimatorCore/Source/PropertyAnimatorCore/Private/Presets/PropertyAnimatorCoreJsonPresetArchive.h
new file mode 100644
--- /dev/null
+++ b/Engine/Plugins/Experimental/PropertyAnimatorCore/Source/PropertyAnimatorCore/Private/Presets/PropertyAnimatorCoreJsonPresetArchive.h
@@ -0,0 +1,129 @@
+// Copyright Epic Games, Inc. All Rights Reserved.
+
+#pragma once
+
+#include "Containers/Array.h"
+#include "Presets/PropertyAnimatorCorePresetArchive.h"
+#include "Templates/SharedPointerFwd.h"
+
+class FJsonObject;
+class FJsonValue;
+
+struct FPropertyAnimatorCorePresetJsonArchiveImplementation : FPropertyAnimatorCorePresetArchiveImplementation
+{
+	static inline const FName Type = TEXT("JSON");
+
+	static TSharedPtr<FPropertyAnimatorCorePresetJsonArchiveImplementation> Get();
+
+	//~ Begin FPropertyAnimatorCorePresetArchiveImplementation
+	virtual TSharedRef<FPropertyAnimatorCorePresetObjectArchive> CreateObject() override;
+	virtual TSharedRef<FPropertyAnimatorCorePresetArrayArchive> CreateArray() override;
+	virtual TSharedRef<FPropertyAnimatorCorePresetValueArchive> CreateValue(bool bInValue) override;
+	virtual TSharedRef<FPropertyAnimatorCorePresetValueArchive> CreateValue(const FString& InValue) override;
+	virtual TSharedRef<FPropertyAnimatorCorePresetValueArchive> CreateValue(uint64 InValue) override;
+	virtual TSharedRef<FPropertyAnimatorCorePresetValueArchive> CreateValue(int64 InValue) override;
+	virtual TSharedRef<FPropertyAnimatorCorePresetValueArchive> CreateValue(double InValue) override;
+	virtual FName GetImplementationType() const override { return Type; }
+	//~ End FPropertyAnimatorCorePresetArchiveImplementation
+
+private:
+	static TSharedPtr<FPropertyAnimatorCorePresetJsonArchiveImplementation> Instance;
+};
+
+struct FPropertyAnimatorCorePresetJsonObjectArchive : FPropertyAnimatorCorePresetObjectArchive
+{
+	FPropertyAnimatorCorePresetJsonObjectArchive();
+	explicit FPropertyAnimatorCorePresetJsonObjectArchive(TSharedRef<FJsonObject> InJsonObject);
+
+	//~ Begin FPropertyAnimatorCorePresetObjectArchive
+	virtual bool Remove(const FString& InKey) override;
+	virtual void Clear() override;
+
+	virtual bool Set(const FString& InKey, TSharedRef<FPropertyAnimatorCorePresetArchive> InValue) override;
+	virtual bool Set(const FString& InKey, bool bInValue) override;
+	virtual bool Set(const FString& InKey, uint64 InValue) override;
+	virtual bool Set(const FString& InKey, int64 InValue) override;
+	virtual bool Set(const FString& InKey, double InValue) override;
+	virtual bool Set(const FString& InKey, const FString& InValue) override;
+
+	virtual bool Get(const FString& InKey, TSharedPtr<FPropertyAnimatorCorePresetArchive>& OutValue) const override;
+	virtual bool Get(const FString& InKey, bool& bOutValue) const override;
+	virtual bool Get(const FString& InKey, uint64& OutValue) const override;
+	virtual bool Get(const FString& InKey, int64& OutValue) const override;
+	virtual bool Get(const FString& InKey, double& OutValue) const override;
+	virtual bool Get(const FString& InKey, FString& OutValue) const override;
+
+	virtual bool ToString(FString& OutString) const override;
+	virtual bool FromString(const FString& InString) override;
+	virtual TSharedRef<FPropertyAnimatorCorePresetArchiveImplementation> GetImplementation() const override;
+	//~ End FPropertyAnimatorCorePresetObjectArchive
+
+	TSharedPtr<FJsonObject> GetJsonObject() const
+	{
+		return JsonObject;
+	}
+
+protected:
+	TSharedPtr<FJsonObject> JsonObject;
+};
+
+struct FPropertyAnimatorCorePresetJsonArrayArchive : FPropertyAnimatorCorePresetArrayArchive
+{
+	FPropertyAnimatorCorePresetJsonArrayArchive() = default;
+	explicit FPropertyAnimatorCorePresetJsonArrayArchive(const TArray<TSharedPtr<FJsonValue>>& InJsonArray);
+
+	//~ Begin FPropertyAnimatorCorePresetArrayArchive
+	virtual bool Get(int32 InIndex, TSharedPtr<FPropertyAnimatorCorePresetArchive>& OutValue) const override;
+	virtual int32 Num() const override;
+	virtual bool Remove(int32 InIndex) override;
+	virtual void Clear() override;
+
+	virtual bool Add(TSharedRef<FPropertyAnimatorCorePresetArchive> InValue) override;
+	virtual bool Add(bool bInValue) override;
+	virtual bool Add(uint64 InValue) override;
+	virtual bool Add(int64 InValue) override;
+	virtual bool Add(double InValue) override;
+	virtual bool Add(const FString& InValue) override;
+
+	virtual bool ToString(FString& OutString) const override;
+	virtual bool FromString(const FString& InString) override;
+	virtual TSharedRef<FPropertyAnimatorCorePresetArchiveImplementation> GetImplementation() const override;
+	//~ End FPropertyAnimatorCorePresetArrayArchive
+
+	const TArray<TSharedPtr<FJsonValue>>& GetJsonValues() const
+	{
+		return JsonValues;
+	}
+
+protected:
+	TArray<TSharedPtr<FJsonValue>> JsonValues;
+};
+
+struct FPropertyAnimatorCorePresetJsonValueArchive : FPropertyAnimatorCorePresetValueArchive
+{
+	explicit FPropertyAnimatorCorePresetJsonValueArchive(bool bInValue);
+	explicit FPropertyAnimatorCorePresetJsonValueArchive(uint64 InValue);
+	explicit FPropertyAnimatorCorePresetJsonValueArchive(int64 InValue);
+	explicit FPropertyAnimatorCorePresetJsonValueArchive(double InValue);
+	explicit FPropertyAnimatorCorePresetJsonValueArchive(const FString& InValue);
+	explicit FPropertyAnimatorCorePresetJsonValueArchive(TSharedRef<FJsonValue> InJsonValue);
+
+	//~ Begin FPropertyAnimatorCorePresetValueArchive
+	virtual bool Get(double& OutValue) const override;
+	virtual bool Get(bool& bOutValue) const override;
+	virtual bool Get(FString& OutValue) const override;
+	virtual bool Get(uint64& OutValue) const override;
+	virtual bool Get(int64& OutValue) const override;
+	virtual bool ToString(FString& OutString) const override;
+	virtual bool FromString(const FString& InString) override;
+	virtual TSharedRef<FPropertyAnimatorCorePresetArchiveImplementation> GetImplementation() const override;
+	//~ End FPropertyAnimatorCorePresetValueArchive
+
+	TSharedPtr<FJsonValue> GetJsonValue() const
+	{
+		return JsonValue;
+	}
+
+protected:
+	TSharedPtr<FJsonValue> JsonValue;
+};
\ No newline at end of file
diff --git a/Engine/Plugins/Experimental/PropertyAnimatorCore/Source/PropertyAnimatorCore/Private/Presets/PropertyAnimatorCorePresetArchive.cpp b/Engine/Plugins/Experimental/PropertyAnimatorCore/Source/PropertyAnimatorCore/Private/Presets/PropertyAnimatorCorePresetArchive.cpp
new file mode 100644
--- /dev/null
+++ b/Engine/Plugins/Experimental/PropertyAnimatorCore/Source/PropertyAnimatorCore/Private/Presets/PropertyAnimatorCorePresetArchive.cpp
@@ -0,0 +1,80 @@
+// Copyright Epic Games, Inc. All Rights Reserved.
+
+#include "Presets/PropertyAnimatorCorePresetArchive.h"
+
+bool FPropertyAnimatorCorePresetArchive::IsObject() const
+{
+	return !!AsObject();
+}
+
+bool FPropertyAnimatorCorePresetArchive::IsArray() const
+{
+	return !!AsArray();
+}
+
+bool FPropertyAnimatorCorePresetArchive::IsValue() const
+{
+	return !!AsValue();
+}
+
+EPropertyAnimatorCorePresetArchiveType FPropertyAnimatorCorePresetArchive::GetType() const
+{
+	if (IsObject())
+	{
+		return EPropertyAnimatorCorePresetArchiveType::Object;
+	}
+
+	if (IsArray())
+	{
+		return EPropertyAnimatorCorePresetArchiveType::Array;
+	}
+
+	if (IsValue())
+	{
+		return EPropertyAnimatorCorePresetArchiveType::Value;
+	}
+
+	check(false);
+	return EPropertyAnimatorCorePresetArchiveType::Value;
+}
+
+FName FPropertyAnimatorCorePresetArchive::GetImplementationType() const
+{
+	return GetImplementation()->GetImplementationType();
+}
+
+bool FPropertyAnimatorCorePresetObjectArchive::Has(const FString& InKey, TOptional<EPropertyAnimatorCorePresetArchiveType> InType) const
+{
+	TSharedPtr<FPropertyAnimatorCorePresetArchive> Value;
+	return Get(InKey, Value) && (!InType.IsSet() || Value->GetType() == InType.GetValue());
+}
+
+TSharedPtr<const FPropertyAnimatorCorePresetObjectArchive> FPropertyAnimatorCorePresetObjectArchive::AsObject() const
+{
+	return SharedThis(this);
+}
+
+TSharedPtr<FPropertyAnimatorCorePresetObjectArchive> FPropertyAnimatorCorePresetObjectArchive::AsMutableObject()
+{
+	return SharedThis(this);
+}
+
+TSharedPtr<const FPropertyAnimatorCorePresetArrayArchive> FPropertyAnimatorCorePresetArrayArchive::AsArray() const
+{
+	return SharedThis(this);
+}
+
+TSharedPtr<FPropertyAnimatorCorePresetArrayArchive> FPropertyAnimatorCorePresetArrayArchive::AsMutableArray()
+{
+	return SharedThis(this);
+}
+
+TSharedPtr<const FPropertyAnimatorCorePresetValueArchive> FPropertyAnimatorCorePresetValueArchive::AsValue() const
+{
+	return SharedThis(this);
+}
+
+TSharedPtr<FPropertyAnimatorCorePresetValueArchive> FPropertyAnimatorCorePresetValueArchive::AsMutableValue()
+{
+	return SharedThis(this);
+}
diff --git a/Engine/Plugins/Experimental/PropertyAnimatorCore/Source/PropertyAnimatorCore/Private/Presets/PropertyAnimatorCorePresetBase.cpp b/Engine/Plugins/Experimental/PropertyAnimatorCore/Source/PropertyAnimatorCore/Private/Presets/PropertyAnimatorCorePresetBase.cpp
--- a/Engine/Plugins/Experimental/PropertyAnimatorCore/Source/PropertyAnimatorCore/Private/Presets/PropertyAnimatorCorePresetBase.cpp
+++ b/Engine/Plugins/Experimental/PropertyAnimatorCore/Source/PropertyAnimatorCore/Private/Presets/PropertyAnimatorCorePresetBase.cpp
@@ -5,6 +5,7 @@
 #include "Animators/PropertyAnimatorCoreBase.h"
 #include "Dom/JsonObject.h"
 #include "Dom/JsonValue.h"
+#include "Presets/PropertyAnimatorCoreJsonPresetArchive.h"
 #include "Presets/PropertyAnimatorCorePresetable.h"
 #include "Serialization/JsonWriter.h"
 #include "Serialization/JsonSerializer.h"
@@ -18,3 +19,8 @@ void UPropertyAnimatorCorePresetBase::CreatePreset(FName InName, const TArray<IP
 {
 	PresetName = InName;
 }
+
+TSharedRef<FPropertyAnimatorCorePresetArchiveImplementation> UPropertyAnimatorCorePresetBase::GetArchiveImplementation() const
+{
+	return FPropertyAnimatorCorePresetJsonArchiveImplementation::Get().ToSharedRef();
+}
diff --git a/Engine/Plugins/Experimental/PropertyAnimatorCore/Source/PropertyAnimatorCore/Private/Presets/PropertyAnimatorCorePropertyPreset.cpp b/Engine/Plugins/Experimental/PropertyAnimatorCore/Source/PropertyAnimatorCore/Private/Presets/PropertyAnimatorCorePropertyPreset.cpp
--- a/Engine/Plugins/Experimental/PropertyAnimatorCore/Source/PropertyAnimatorCore/Private/Presets/PropertyAnimatorCorePropertyPreset.cpp
+++ b/Engine/Plugins/Experimental/PropertyAnimatorCore/Source/PropertyAnimatorCore/Private/Presets/PropertyAnimatorCorePropertyPreset.cpp
@@ -3,16 +3,11 @@
 #include "Presets/PropertyAnimatorCorePropertyPreset.h"
 
 #include "Animators/PropertyAnimatorCoreBase.h"
-#include "Dom/JsonObject.h"
-#include "Dom/JsonValue.h"
 #include "Properties/PropertyAnimatorCoreContext.h"
-#include "Serialization/JsonReader.h"
-#include "Serialization/JsonSerializer.h"
-#include "Serialization/JsonWriter.h"
 
 void UPropertyAnimatorCorePropertyPreset::GetPresetProperties(const AActor* InActor, const UPropertyAnimatorCoreBase* InAnimator, TSet<FPropertyAnimatorCoreData>& OutProperties) const
 {
-	for (const TPair<FString, TSharedRef<FJsonValue>>& PropertyPreset : PropertyPresets)
+	for (const TPair<FString, TSharedRef<FPropertyAnimatorCorePresetArchive>>& PropertyPreset : PropertyPresets)
 	{
 		FPropertyAnimatorCoreData Property(const_cast<AActor*>(InActor), PropertyPreset.Key);
 
@@ -25,7 +20,7 @@ void UPropertyAnimatorCorePropertyPreset::GetPresetProperties(const AActor* InAc
 
 void UPropertyAnimatorCorePropertyPreset::OnPresetApplied(UPropertyAnimatorCoreBase* InAnimator, const TSet<FPropertyAnimatorCoreData>& InProperties)
 {
-	for (const TPair<FString, TSharedRef<FJsonValue>>& PropertyPreset : PropertyPresets)
+	for (const TPair<FString, TSharedRef<FPropertyAnimatorCorePresetArchive>>& PropertyPreset : PropertyPresets)
 	{
 		FPropertyAnimatorCoreData Property(InAnimator->GetAnimatorActor(), PropertyPreset.Key);
 
@@ -163,26 +158,26 @@ void UPropertyAnimatorCorePropertyPreset::CreatePreset(FName InName, const TArra
 {
 	Super::CreatePreset(InName, InPresetableItems);
 
-	TArray<TSharedPtr<FJsonValue>> JsonValues;
+	TSharedPtr<FPropertyAnimatorCorePresetArrayArchive> PropertiesArchive = GetArchiveImplementation()->CreateArray();
 
 	for (IPropertyAnimatorCorePresetable* InPresetableItem : InPresetableItems)
 	{
-		TSharedPtr<FJsonValue> JsonValue;
+		TSharedPtr<FPropertyAnimatorCorePresetArchive> PropertyArchive;
 
 		if (InPresetableItem
-			&& InPresetableItem->ExportPreset(this, JsonValue)
-			&& JsonValue.IsValid())
+			&& InPresetableItem->ExportPreset(this, PropertyArchive)
+			&& PropertyArchive.IsValid())
 		{
-			JsonValues.Add(JsonValue);
+			PropertiesArchive->Add(PropertyArchive.ToSharedRef());
 		}
 	}
 
-	FString JsonString;
-	TSharedRef<TJsonWriter<>> JsonWriter = TJsonWriterFactory<>::Create(&JsonString);
-	if (FJsonSerializer::Serialize(JsonValues, JsonWriter))
+	FString OutputString;
+	if (PropertiesArchive->ToString(OutputString))
 	{
 		PresetVersion = 0;
-		PresetContent = JsonString;
+		PresetFormat = PropertiesArchive->GetImplementationType();
+		PresetContent = OutputString;
 	}
 }
 
@@ -193,30 +188,27 @@ bool UPropertyAnimatorCorePropertyPreset::LoadPreset()
 		return false;
 	}
 
-	TSharedRef<TJsonReader<>> Reader = TJsonReaderFactory<>::Create(PresetContent);
+	TSharedPtr<FPropertyAnimatorCorePresetArrayArchive> PropertiesArchive = GetArchiveImplementation()->CreateArray();
 
-	TArray<TSharedPtr<FJsonValue>> JsonArray;
-	if (!FJsonSerializer::Deserialize(Reader, JsonArray) || JsonArray.IsEmpty())
+	if (!PropertiesArchive->FromString(PresetContent) || PropertiesArchive->Num() == 0)
 	{
 		return false;
 	}
 
-	for (const TSharedPtr<FJsonValue>& JsonValue : JsonArray)
+	for (int32 Index = 0; Index < PropertiesArchive->Num(); Index++)
 	{
-		TSharedPtr<FJsonObject>* JsonObject = nullptr;
-		if (!JsonValue
-			|| !JsonValue->TryGetObject(JsonObject)
-			|| !JsonObject)
+		TSharedPtr<FPropertyAnimatorCorePresetArchive> PropertyArchive;
+		if (!PropertiesArchive->Get(Index, PropertyArchive) || !PropertyArchive->IsObject())
 		{
 			continue;
 		}
 
-		FString PropertyPath;
+		const TSharedPtr<FPropertyAnimatorCorePresetObjectArchive> PropertyObject = PropertyArchive->AsMutableObject();
 
-		if ((*JsonObject)->TryGetStringField(UPropertyAnimatorCoreContext::GetAnimatedPropertyName().ToString(), PropertyPath)
-			&& !PropertyPath.IsEmpty())
+		FString PropertyPath;
+		if (PropertyObject->Get(UPropertyAnimatorCoreContext::GetAnimatedPropertyName().ToString(), PropertyPath) && !PropertyPath.IsEmpty())
 		{
-			PropertyPresets.Add(PropertyPath, JsonValue.ToSharedRef());
+			PropertyPresets.Add(PropertyPath, PropertyArchive.ToSharedRef());
 		}
 	}
 
diff --git a/Engine/Plugins/Experimental/PropertyAnimatorCore/Source/PropertyAnimatorCore/Private/Properties/PropertyAnimatorCoreContext.cpp b/Engine/Plugins/Experimental/PropertyAnimatorCore/Source/PropertyAnimatorCore/Private/Properties/PropertyAnimatorCoreContext.cpp
--- a/Engine/Plugins/Experimental/PropertyAnimatorCore/Source/PropertyAnimatorCore/Private/Properties/PropertyAnimatorCoreContext.cpp
+++ b/Engine/Plugins/Experimental/PropertyAnimatorCore/Source/PropertyAnimatorCore/Private/Properties/PropertyAnimatorCoreContext.cpp
@@ -3,8 +3,6 @@
 #include "Properties/PropertyAnimatorCoreContext.h"
 
 #include "Containers/Ticker.h"
-#include "Dom/JsonObject.h"
-#include "Dom/JsonValue.h"
 #include "StructUtils/InstancedStruct.h"
 #include "Properties/PropertyAnimatorCoreResolver.h"
 #include "Properties/Converters/PropertyAnimatorCoreConverterBase.h"
@@ -180,73 +178,76 @@ void UPropertyAnimatorCoreContext::PostEditChangeProperty(FPropertyChangedEvent&
 }
 #endif
 
-bool UPropertyAnimatorCoreContext::ImportPreset(const UPropertyAnimatorCorePresetBase* InPreset, const TSharedRef<FJsonValue>& InValue)
+bool UPropertyAnimatorCoreContext::ImportPreset(const UPropertyAnimatorCorePresetBase* InPreset, const TSharedRef<FPropertyAnimatorCorePresetArchive>& InValue)
 {
-	TSharedPtr<FJsonObject>* JsonObject;
-	if (!InValue->TryGetObject(JsonObject) || !JsonObject || !JsonObject->IsValid())
+	TSharedPtr<FPropertyAnimatorCorePresetObjectArchive> ObjectArchive = InValue->AsMutableObject();
+
+	if (!ObjectArchive)
 	{
 		return false;
 	}
 
-	bool bJsonAnimated = bAnimated;
-	(*JsonObject)->TryGetBoolField(GET_MEMBER_NAME_STRING_CHECKED(UPropertyAnimatorCoreContext, bAnimated), bJsonAnimated);
-	SetAnimated(bJsonAnimated);
+	bool bAnimatedArchive = bAnimated;
+	ObjectArchive->Get(GET_MEMBER_NAME_STRING_CHECKED(UPropertyAnimatorCoreContext, bAnimated), bAnimatedArchive);
+	SetAnimated(bAnimatedArchive);
 
 	if (bEditMagnitude)
 	{
-		float JsonMagnitude = Magnitude;
-		(*JsonObject)->TryGetNumberField(GET_MEMBER_NAME_STRING_CHECKED(UPropertyAnimatorCoreContext, Magnitude), JsonMagnitude);
-		SetMagnitude(JsonMagnitude);
+		double MagnitudeArchive = Magnitude;
+		ObjectArchive->Get(GET_MEMBER_NAME_STRING_CHECKED(UPropertyAnimatorCoreContext, Magnitude), MagnitudeArchive);
+		SetMagnitude(MagnitudeArchive);
 
-		float JsonTimeOffset = TimeOffset;
-		(*JsonObject)->TryGetNumberField(GET_MEMBER_NAME_STRING_CHECKED(UPropertyAnimatorCoreContext, TimeOffset), JsonTimeOffset);
-		SetTimeOffset(JsonTimeOffset);
+		double TimeOffsetArchive = TimeOffset;
+		ObjectArchive->Get(GET_MEMBER_NAME_STRING_CHECKED(UPropertyAnimatorCoreContext, TimeOffset), TimeOffsetArchive);
+		SetTimeOffset(TimeOffsetArchive);
 	}
 
 	if (bEditMode)
 	{
-		float JsonMode = static_cast<float>(Mode);
-		(*JsonObject)->TryGetNumberField(GET_MEMBER_NAME_STRING_CHECKED(UPropertyAnimatorCoreContext, Mode), JsonMode);
-		SetMode(static_cast<EPropertyAnimatorCoreMode>(JsonMode));
+		int64 ModeArchive = static_cast<int64>(Mode);
+		ObjectArchive->Get(GET_MEMBER_NAME_STRING_CHECKED(UPropertyAnimatorCoreContext, Mode), ModeArchive);
+		SetMode(static_cast<EPropertyAnimatorCoreMode>(ModeArchive));
 	}
 
 	if (Resolver)
 	{
-		if ((*JsonObject)->HasTypedField(GET_MEMBER_NAME_STRING_CHECKED(UPropertyAnimatorCoreContext, Resolver), EJson::Object))
+		if (ObjectArchive->Has(GET_MEMBER_NAME_STRING_CHECKED(UPropertyAnimatorCoreContext, Resolver), EPropertyAnimatorCorePresetArchiveType::Object))
 		{
-			Resolver->ImportPreset(InPreset, (*JsonObject)->TryGetField(GET_MEMBER_NAME_STRING_CHECKED(UPropertyAnimatorCoreContext, Resolver)).ToSharedRef());
+			TSharedPtr<FPropertyAnimatorCorePresetArchive> ResolverArchive;
+			ObjectArchive->Get(GET_MEMBER_NAME_STRING_CHECKED(UPropertyAnimatorCoreContext, Resolver), ResolverArchive);
+			Resolver->ImportPreset(InPreset, ResolverArchive.ToSharedRef());
 		}
 	}
 
 	return true;
 }
 
-bool UPropertyAnimatorCoreContext::ExportPreset(const UPropertyAnimatorCorePresetBase* InPreset, TSharedPtr<FJsonValue>& OutValue)
+bool UPropertyAnimatorCoreContext::ExportPreset(const UPropertyAnimatorCorePresetBase* InPreset, TSharedPtr<FPropertyAnimatorCorePresetArchive>& OutValue) const
 {
-	TSharedRef<FJsonObject> JsonObject = MakeShared<FJsonObject>();
-	OutValue = MakeShared<FJsonValueObject>(JsonObject);
+	TSharedRef<FPropertyAnimatorCorePresetObjectArchive> ContextArchive = InPreset->GetArchiveImplementation()->CreateObject();
+	OutValue = ContextArchive;
 
-	JsonObject->SetBoolField(GET_MEMBER_NAME_STRING_CHECKED(UPropertyAnimatorCoreContext, bAnimated), bAnimated);
+	ContextArchive->Set(GET_MEMBER_NAME_STRING_CHECKED(UPropertyAnimatorCoreContext, bAnimated), bAnimated);
 
 	if (bEditMagnitude)
 	{
-		JsonObject->SetNumberField(GET_MEMBER_NAME_STRING_CHECKED(UPropertyAnimatorCoreContext, Magnitude), Magnitude);
-		JsonObject->SetNumberField(GET_MEMBER_NAME_STRING_CHECKED(UPropertyAnimatorCoreContext, TimeOffset), TimeOffset);
+		ContextArchive->Set(GET_MEMBER_NAME_STRING_CHECKED(UPropertyAnimatorCoreContext, Magnitude), Magnitude);
+		ContextArchive->Set(GET_MEMBER_NAME_STRING_CHECKED(UPropertyAnimatorCoreContext, TimeOffset), TimeOffset);
 	}
 
 	if (bEditMode)
 	{
-		JsonObject->SetNumberField(GET_MEMBER_NAME_STRING_CHECKED(UPropertyAnimatorCoreContext, Mode), static_cast<float>(Mode));
+		ContextArchive->Set(GET_MEMBER_NAME_STRING_CHECKED(UPropertyAnimatorCoreContext, Mode), static_cast<uint64>(Mode));
 	}
 
-	JsonObject->SetStringField(GET_MEMBER_NAME_STRING_CHECKED(UPropertyAnimatorCoreContext, AnimatedProperty), AnimatedProperty.GetPropertyLocatorPath());
+	ContextArchive->Set(GET_MEMBER_NAME_STRING_CHECKED(UPropertyAnimatorCoreContext, AnimatedProperty), AnimatedProperty.GetPropertyLocatorPath());
 
 	if (Resolver)
 	{
-		TSharedPtr<FJsonValue> JsonValue;
-		if (Resolver->ExportPreset(InPreset, JsonValue) && JsonValue.IsValid())
+		TSharedPtr<FPropertyAnimatorCorePresetArchive> ResolverArchive;
+		if (Resolver->ExportPreset(InPreset, ResolverArchive) && ResolverArchive.IsValid())
 		{
-			JsonObject->SetField(GET_MEMBER_NAME_STRING_CHECKED(UPropertyAnimatorCoreContext, Resolver), JsonValue);
+			ContextArchive->Set(GET_MEMBER_NAME_STRING_CHECKED(UPropertyAnimatorCoreContext, Resolver), ResolverArchive.ToSharedRef());
 		}
 	}
 
@@ -290,12 +291,12 @@ bool UPropertyAnimatorCoreContext::ResolvePropertyOwner(AActor* InNewOwner)
 		return true;
 	}
 
-	bool bFound = IsValid(NewOwner);
-	const TArray<UObject*> OtherOuters = AnimatedProperty.GetOuters(AnimatedProperty.GetOwningActor());
-
 	const FProperty* MemberProperty = AnimatedProperty.GetMemberProperty();
 	UClass* PropertyOwningClass = MemberProperty->GetOwnerClass();
 
+	bool bFound = IsValid(NewOwner);
+	const TArray<UObject*> OtherOuters = AnimatedProperty.GetOuters(AnimatedProperty.GetOwningActor());
+
 	if (!OtherOuters.IsEmpty())
 	{
 		// Resolve using outers
diff --git a/Engine/Plugins/Experimental/PropertyAnimatorCore/Source/PropertyAnimatorCore/Private/Properties/PropertyAnimatorCoreResolver.cpp b/Engine/Plugins/Experimental/PropertyAnimatorCore/Source/PropertyAnimatorCore/Private/Properties/PropertyAnimatorCoreResolver.cpp
--- a/Engine/Plugins/Experimental/PropertyAnimatorCore/Source/PropertyAnimatorCore/Private/Properties/PropertyAnimatorCoreResolver.cpp
+++ b/Engine/Plugins/Experimental/PropertyAnimatorCore/Source/PropertyAnimatorCore/Private/Properties/PropertyAnimatorCoreResolver.cpp
@@ -2,13 +2,12 @@
 
 #include "Properties/PropertyAnimatorCoreResolver.h"
 
-#include "Dom/JsonObject.h"
-#include "Dom/JsonValue.h"
+#include "Presets/PropertyAnimatorCorePresetArchive.h"
+#include "Presets/PropertyAnimatorCorePresetBase.h"
 
-bool UPropertyAnimatorCoreResolver::ImportPreset(const UPropertyAnimatorCorePresetBase* InPreset, const TSharedRef<FJsonValue>& InValue)
+bool UPropertyAnimatorCoreResolver::ImportPreset(const UPropertyAnimatorCorePresetBase* InPreset, const TSharedRef<FPropertyAnimatorCorePresetArchive>& InValue)
 {
-	const TSharedPtr<FJsonObject>* JsonTimeSourceObject;
-	if (!InValue->TryGetObject(JsonTimeSourceObject) || !JsonTimeSourceObject)
+	if (!InValue->IsObject())
 	{
 		return false;
 	}
@@ -16,10 +15,9 @@ bool UPropertyAnimatorCoreResolver::ImportPreset(const UPropertyAnimatorCorePres
 	return true;
 }
 
-bool UPropertyAnimatorCoreResolver::ExportPreset(const UPropertyAnimatorCorePresetBase* InPreset, TSharedPtr<FJsonValue>& OutValue)
+bool UPropertyAnimatorCoreResolver::ExportPreset(const UPropertyAnimatorCorePresetBase* InPreset, TSharedPtr<FPropertyAnimatorCorePresetArchive>& OutValue) const
 {
-	TSharedRef<FJsonObject> JsonTimeSourceObject = MakeShared<FJsonObject>();
-	OutValue = MakeShared<FJsonValueObject>(JsonTimeSourceObject);
-
+	const TSharedRef<FPropertyAnimatorCorePresetObjectArchive> TimeSourceArchive = InPreset->GetArchiveImplementation()->CreateObject();
+	OutValue = TimeSourceArchive;
 	return true;
 }
diff --git a/Engine/Plugins/Experimental/PropertyAnimatorCore/Source/PropertyAnimatorCore/Private/TimeSources/PropertyAnimatorCoreManualTimeSource.cpp b/Engine/Plugins/Experimental/PropertyAnimatorCore/Source/PropertyAnimatorCore/Private/TimeSources/PropertyAnimatorCoreManualTimeSource.cpp
--- a/Engine/Plugins/Experimental/PropertyAnimatorCore/Source/PropertyAnimatorCore/Private/TimeSources/PropertyAnimatorCoreManualTimeSource.cpp
+++ b/Engine/Plugins/Experimental/PropertyAnimatorCore/Source/PropertyAnimatorCore/Private/TimeSources/PropertyAnimatorCoreManualTimeSource.cpp
@@ -2,9 +2,8 @@
 
 #include "TimeSources/PropertyAnimatorCoreManualTimeSource.h"
 
-#include "Dom/JsonObject.h"
-#include "Dom/JsonValue.h"
 #include "Misc/App.h"
+#include "Presets/PropertyAnimatorCorePresetArchive.h"
 
 void UPropertyAnimatorCoreManualTimeSource::SetOverrideTime(bool bInOverride)
 {
@@ -79,19 +78,19 @@ void UPropertyAnimatorCoreManualTimeSource::OnTimeSourceInactive()
 	Stop();
 }
 
-bool UPropertyAnimatorCoreManualTimeSource::ImportPreset(const UPropertyAnimatorCorePresetBase* InPreset, const TSharedRef<FJsonValue>& InValue)
+bool UPropertyAnimatorCoreManualTimeSource::ImportPreset(const UPropertyAnimatorCorePresetBase* InPreset, const TSharedRef<FPropertyAnimatorCorePresetArchive>& InValue)
 {
-	const TSharedPtr<FJsonObject>* JsonTimeSourceObject;
-
-	if (Super::ImportPreset(InPreset, InValue) && InValue->TryGetObject(JsonTimeSourceObject))
+	if (Super::ImportPreset(InPreset, InValue) && InValue->IsObject())
 	{
-		double JsonCustomTime = CustomTime;
-		(*JsonTimeSourceObject)->TryGetNumberField(GET_MEMBER_NAME_STRING_CHECKED(UPropertyAnimatorCoreManualTimeSource, CustomTime), JsonCustomTime);
-		SetCustomTime(JsonCustomTime);
+		TSharedPtr<FPropertyAnimatorCorePresetObjectArchive> ObjectArchive = InValue->AsMutableObject();
+
+		double CustomTimeValue = CustomTime;
+		ObjectArchive->Get(GET_MEMBER_NAME_STRING_CHECKED(UPropertyAnimatorCoreManualTimeSource, CustomTime), CustomTimeValue);
+		SetCustomTime(CustomTimeValue);
 
-		bool bJsonOverrideTime = bOverrideTime;
-		(*JsonTimeSourceObject)->TryGetBoolField(GET_MEMBER_NAME_STRING_CHECKED(UPropertyAnimatorCoreManualTimeSource, bOverrideTime), bJsonOverrideTime);
-		SetOverrideTime(bJsonOverrideTime);
+		bool bOverrideTimeValue = bOverrideTime;
+		ObjectArchive->Get(GET_MEMBER_NAME_STRING_CHECKED(UPropertyAnimatorCoreManualTimeSource, bOverrideTime), bOverrideTimeValue);
+		SetOverrideTime(bOverrideTimeValue);
 
 		return true;
 	}
@@ -99,14 +98,14 @@ bool UPropertyAnimatorCoreManualTimeSource::ImportPreset(const UPropertyAnimator
 	return false;
 }
 
-bool UPropertyAnimatorCoreManualTimeSource::ExportPreset(const UPropertyAnimatorCorePresetBase* InPreset, TSharedPtr<FJsonValue>& OutValue)
+bool UPropertyAnimatorCoreManualTimeSource::ExportPreset(const UPropertyAnimatorCorePresetBase* InPreset, TSharedPtr<FPropertyAnimatorCorePresetArchive>& OutValue) const
 {
-	const TSharedPtr<FJsonObject>* JsonTimeSourceObject;
-
-	if (Super::ExportPreset(InPreset, OutValue) && OutValue->TryGetObject(JsonTimeSourceObject))
+	if (Super::ExportPreset(InPreset, OutValue) && OutValue->IsObject())
 	{
-		(*JsonTimeSourceObject)->SetNumberField(GET_MEMBER_NAME_STRING_CHECKED(UPropertyAnimatorCoreManualTimeSource, CustomTime), CustomTime);
-		(*JsonTimeSourceObject)->SetBoolField(GET_MEMBER_NAME_STRING_CHECKED(UPropertyAnimatorCoreManualTimeSource, bOverrideTime), bOverrideTime);
+		TSharedPtr<FPropertyAnimatorCorePresetObjectArchive> ObjectArchive = OutValue->AsMutableObject();
+
+		ObjectArchive->Set(GET_MEMBER_NAME_STRING_CHECKED(UPropertyAnimatorCoreManualTimeSource, CustomTime), CustomTime);
+		ObjectArchive->Set(GET_MEMBER_NAME_STRING_CHECKED(UPropertyAnimatorCoreManualTimeSource, bOverrideTime), bOverrideTime);
 
 		return true;
 	}
diff --git a/Engine/Plugins/Experimental/PropertyAnimatorCore/Source/PropertyAnimatorCore/Private/TimeSources/PropertyAnimatorCoreSystemTimeSource.cpp b/Engine/Plugins/Experimental/PropertyAnimatorCore/Source/PropertyAnimatorCore/Private/TimeSources/PropertyAnimatorCoreSystemTimeSource.cpp
--- a/Engine/Plugins/Experimental/PropertyAnimatorCore/Source/PropertyAnimatorCore/Private/TimeSources/PropertyAnimatorCoreSystemTimeSource.cpp
+++ b/Engine/Plugins/Experimental/PropertyAnimatorCore/Source/PropertyAnimatorCore/Private/TimeSources/PropertyAnimatorCoreSystemTimeSource.cpp
@@ -2,10 +2,9 @@
 
 #include "TimeSources/PropertyAnimatorCoreSystemTimeSource.h"
 
-#include "Dom/JsonObject.h"
-#include "Dom/JsonValue.h"
 #include "Internationalization/Regex.h"
 #include "Misc/DateTime.h"
+#include "Presets/PropertyAnimatorCorePresetArchive.h"
 
 bool UPropertyAnimatorCoreSystemTimeSource::UpdateEvaluationData(FPropertyAnimatorCoreTimeSourceEvaluationData& OutData)
 {
@@ -44,19 +43,19 @@ void UPropertyAnimatorCoreSystemTimeSource::OnTimeSourceActive()
 	OnModeChanged();
 }
 
-bool UPropertyAnimatorCoreSystemTimeSource::ImportPreset(const UPropertyAnimatorCorePresetBase* InPreset, const TSharedRef<FJsonValue>& InValue)
+bool UPropertyAnimatorCoreSystemTimeSource::ImportPreset(const UPropertyAnimatorCorePresetBase* InPreset, const TSharedRef<FPropertyAnimatorCorePresetArchive>& InValue)
 {
-	const TSharedPtr<FJsonObject>* JsonTimeSourceObject;
-
-	if (Super::ImportPreset(InPreset, InValue) && InValue->TryGetObject(JsonTimeSourceObject))
+	if (Super::ImportPreset(InPreset, InValue) && InValue->IsObject())
 	{
-		FString JsonCountdown = CountdownDuration;
-		(*JsonTimeSourceObject)->TryGetStringField(GET_MEMBER_NAME_STRING_CHECKED(UPropertyAnimatorCoreSystemTimeSource, CountdownDuration), JsonCountdown);
-		SetCountdownDuration(JsonCountdown);
+		const TSharedPtr<FPropertyAnimatorCorePresetObjectArchive> ObjectArchive = InValue->AsMutableObject();
+
+		FString CountdownValue = CountdownDuration;
+		ObjectArchive->Get(GET_MEMBER_NAME_STRING_CHECKED(UPropertyAnimatorCoreSystemTimeSource, CountdownDuration), CountdownValue);
+		SetCountdownDuration(CountdownValue);
 
-		double JsonMode = static_cast<double>(Mode);
-		(*JsonTimeSourceObject)->TryGetNumberField(GET_MEMBER_NAME_STRING_CHECKED(UPropertyAnimatorCoreSystemTimeSource, Mode), JsonMode);
-		SetMode(static_cast<EPropertyAnimatorCoreSystemMode>(JsonMode));
+		uint64 ModeValue = static_cast<uint64>(Mode);
+		ObjectArchive->Get(GET_MEMBER_NAME_STRING_CHECKED(UPropertyAnimatorCoreSystemTimeSource, Mode), ModeValue);
+		SetMode(static_cast<EPropertyAnimatorCoreSystemMode>(ModeValue));
 
 		return true;
 	}
@@ -64,14 +63,14 @@ bool UPropertyAnimatorCoreSystemTimeSource::ImportPreset(const UPropertyAnimator
 	return false;
 }
 
-bool UPropertyAnimatorCoreSystemTimeSource::ExportPreset(const UPropertyAnimatorCorePresetBase* InPreset, TSharedPtr<FJsonValue>& OutValue)
+bool UPropertyAnimatorCoreSystemTimeSource::ExportPreset(const UPropertyAnimatorCorePresetBase* InPreset, TSharedPtr<FPropertyAnimatorCorePresetArchive>& OutValue) const
 {
-	TSharedPtr<FJsonObject>* JsonTimeSourceObject = nullptr;
-
-	if (Super::ExportPreset(InPreset, OutValue) && OutValue->TryGetObject(JsonTimeSourceObject))
+	if (Super::ExportPreset(InPreset, OutValue) && OutValue->IsObject())
 	{
-		(*JsonTimeSourceObject)->SetNumberField(GET_MEMBER_NAME_STRING_CHECKED(UPropertyAnimatorCoreSystemTimeSource, Mode), static_cast<double>(Mode));
-		(*JsonTimeSourceObject)->SetStringField(GET_MEMBER_NAME_STRING_CHECKED(UPropertyAnimatorCoreSystemTimeSource, CountdownDuration), CountdownDuration);
+		const TSharedPtr<FPropertyAnimatorCorePresetObjectArchive> ObjectArchive = OutValue->AsMutableObject();
+
+		ObjectArchive->Set(GET_MEMBER_NAME_STRING_CHECKED(UPropertyAnimatorCoreSystemTimeSource, Mode), static_cast<uint64>(Mode));
+		ObjectArchive->Set(GET_MEMBER_NAME_STRING_CHECKED(UPropertyAnimatorCoreSystemTimeSource, CountdownDuration), CountdownDuration);
 
 		return true;
 	}
diff --git a/Engine/Plugins/Experimental/PropertyAnimatorCore/Source/PropertyAnimatorCore/Private/TimeSources/PropertyAnimatorCoreTimeSourceBase.cpp b/Engine/Plugins/Experimental/PropertyAnimatorCore/Source/PropertyAnimatorCore/Private/TimeSources/PropertyAnimatorCoreTimeSourceBase.cpp
--- a/Engine/Plugins/Experimental/PropertyAnimatorCore/Source/PropertyAnimatorCore/Private/TimeSources/PropertyAnimatorCoreTimeSourceBase.cpp
+++ b/Engine/Plugins/Experimental/PropertyAnimatorCore/Source/PropertyAnimatorCore/Private/TimeSources/PropertyAnimatorCoreTimeSourceBase.cpp
@@ -2,8 +2,8 @@
 
 #include "TimeSources/PropertyAnimatorCoreTimeSourceBase.h"
 
-#include "Dom/JsonObject.h"
-#include "Dom/JsonValue.h"
+#include "Presets/PropertyAnimatorCorePresetArchive.h"
+#include "Presets/PropertyAnimatorCorePresetBase.h"
 
 void UPropertyAnimatorCoreTimeSourceBase::ActivateTimeSource()
 {
@@ -56,32 +56,33 @@ void UPropertyAnimatorCoreTimeSourceBase::SetUseFrameRate(bool bInUseFrameRate)
 	bUseFrameRate = bInUseFrameRate;
 }
 
-bool UPropertyAnimatorCoreTimeSourceBase::ImportPreset(const UPropertyAnimatorCorePresetBase* InPreset, const TSharedRef<FJsonValue>& InValue)
+bool UPropertyAnimatorCoreTimeSourceBase::ImportPreset(const UPropertyAnimatorCorePresetBase* InPreset, const TSharedRef<FPropertyAnimatorCorePresetArchive>& InValue)
 {
-	const TSharedPtr<FJsonObject>* JsonTimeSourceObject;
-	if (!InValue->TryGetObject(JsonTimeSourceObject) || !JsonTimeSourceObject)
+	TSharedPtr<FPropertyAnimatorCorePresetObjectArchive> ObjectArchive = InValue->AsMutableObject();
+
+	if (!ObjectArchive)
 	{
 		return false;
 	}
 
-	bool bJsonUseFrameRate = bUseFrameRate;
-	(*JsonTimeSourceObject)->TryGetBoolField(GET_MEMBER_NAME_STRING_CHECKED(UPropertyAnimatorCoreTimeSourceBase, bUseFrameRate), bJsonUseFrameRate);
-	SetUseFrameRate(bJsonUseFrameRate);
+	bool bUseFrameRateValue = bUseFrameRate;
+	ObjectArchive->Get(GET_MEMBER_NAME_STRING_CHECKED(UPropertyAnimatorCoreTimeSourceBase, bUseFrameRate), bUseFrameRateValue);
+	SetUseFrameRate(bUseFrameRateValue);
 
-	double JsonFrameRate = FrameRate;
-	(*JsonTimeSourceObject)->TryGetNumberField(GET_MEMBER_NAME_STRING_CHECKED(UPropertyAnimatorCoreTimeSourceBase, FrameRate), JsonFrameRate);
-	SetFrameRate(JsonFrameRate);
+	double FrameRateValue = FrameRate;
+	ObjectArchive->Get(GET_MEMBER_NAME_STRING_CHECKED(UPropertyAnimatorCoreTimeSourceBase, FrameRate), FrameRateValue);
+	SetFrameRate(FrameRateValue);
 
 	return true;
 }
 
-bool UPropertyAnimatorCoreTimeSourceBase::ExportPreset(const UPropertyAnimatorCorePresetBase* InPreset, TSharedPtr<FJsonValue>& OutValue)
+bool UPropertyAnimatorCoreTimeSourceBase::ExportPreset(const UPropertyAnimatorCorePresetBase* InPreset, TSharedPtr<FPropertyAnimatorCorePresetArchive>& OutValue) const
 {
-	TSharedRef<FJsonObject> JsonTimeSourceObject = MakeShared<FJsonObject>();
-	OutValue = MakeShared<FJsonValueObject>(JsonTimeSourceObject);
+	const TSharedRef<FPropertyAnimatorCorePresetObjectArchive> ObjectArchive = InPreset->GetArchiveImplementation()->CreateObject();
+	OutValue = ObjectArchive;
 
-	JsonTimeSourceObject->SetBoolField(GET_MEMBER_NAME_STRING_CHECKED(UPropertyAnimatorCoreTimeSourceBase, bUseFrameRate), bUseFrameRate);
-	JsonTimeSourceObject->SetNumberField(GET_MEMBER_NAME_STRING_CHECKED(UPropertyAnimatorCoreTimeSourceBase, FrameRate), FrameRate);
+	ObjectArchive->Set(GET_MEMBER_NAME_STRING_CHECKED(UPropertyAnimatorCoreTimeSourceBase, bUseFrameRate), bUseFrameRate);
+	ObjectArchive->Set(GET_MEMBER_NAME_STRING_CHECKED(UPropertyAnimatorCoreTimeSourceBase, FrameRate), bUseFrameRate);
 
 	return true;
 }
diff --git a/Engine/Plugins/Experimental/PropertyAnimatorCore/Source/PropertyAnimatorCore/Public/Animators/PropertyAnimatorCoreBase.h b/Engine/Plugins/Experimental/PropertyAnimatorCore/Source/PropertyAnimatorCore/Public/Animators/PropertyAnimatorCoreBase.h
--- a/Engine/Plugins/Experimental/PropertyAnimatorCore/Source/PropertyAnimatorCore/Public/Animators/PropertyAnimatorCoreBase.h
+++ b/Engine/Plugins/Experimental/PropertyAnimatorCore/Source/PropertyAnimatorCore/Public/Animators/PropertyAnimatorCoreBase.h
@@ -215,8 +215,8 @@ public:
 	}
 
 	//~ Begin IPropertyAnimatorCorePresetable
-	PROPERTYANIMATORCORE_API virtual bool ImportPreset(const UPropertyAnimatorCorePresetBase* InPreset, const TSharedRef<FJsonValue>& InValue) override;
-	PROPERTYANIMATORCORE_API virtual bool ExportPreset(const UPropertyAnimatorCorePresetBase* InPreset, TSharedPtr<FJsonValue>& OutValue) override;
+	PROPERTYANIMATORCORE_API virtual bool ImportPreset(const UPropertyAnimatorCorePresetBase* InPreset, const TSharedRef<FPropertyAnimatorCorePresetArchive>& InValue) override;
+	PROPERTYANIMATORCORE_API virtual bool ExportPreset(const UPropertyAnimatorCorePresetBase* InPreset, TSharedPtr<FPropertyAnimatorCorePresetArchive>& OutValue) const override;
 	//~ End IPropertyAnimatorCorePresetable
 
 protected:
diff --git a/Engine/Plugins/Experimental/PropertyAnimatorCore/Source/PropertyAnimatorCore/Public/Presets/PropertyAnimatorCoreAnimatorPreset.h b/Engine/Plugins/Experimental/PropertyAnimatorCore/Source/PropertyAnimatorCore/Public/Presets/PropertyAnimatorCoreAnimatorPreset.h
--- a/Engine/Plugins/Experimental/PropertyAnimatorCore/Source/PropertyAnimatorCore/Public/Presets/PropertyAnimatorCoreAnimatorPreset.h
+++ b/Engine/Plugins/Experimental/PropertyAnimatorCore/Source/PropertyAnimatorCore/Public/Presets/PropertyAnimatorCoreAnimatorPreset.h
@@ -34,5 +34,5 @@ protected:
 	UPROPERTY(Transient)
 	TSubclassOf<UPropertyAnimatorCoreBase> TargetAnimatorClass;
 
-	TSharedPtr<FJsonValue> AnimatorPreset;
+	TSharedPtr<FPropertyAnimatorCorePresetArchive> AnimatorPreset;
 };
diff --git a/Engine/Plugins/Experimental/PropertyAnimatorCore/Source/PropertyAnimatorCore/Public/Presets/PropertyAnimatorCorePresetArchive.h b/Engine/Plugins/Experimental/PropertyAnimatorCore/Source/PropertyAnimatorCore/Public/Presets/PropertyAnimatorCorePresetArchive.h
new file mode 100644
--- /dev/null
+++ b/Engine/Plugins/Experimental/PropertyAnimatorCore/Source/PropertyAnimatorCore/Public/Presets/PropertyAnimatorCorePresetArchive.h
@@ -0,0 +1,126 @@
+// Copyright Epic Games, Inc. All Rights Reserved.
+
+#pragma once
+
+#include "HAL/Platform.h"
+#include "Misc/Optional.h"
+#include "Templates/SharedPointer.h"
+
+struct FPropertyAnimatorCorePresetArchive;
+struct FPropertyAnimatorCorePresetObjectArchive;
+struct FPropertyAnimatorCorePresetArrayArchive;
+struct FPropertyAnimatorCorePresetValueArchive;
+
+/** Enumerates all possible value type stored in archive */
+enum class EPropertyAnimatorCorePresetArchiveType : uint8
+{
+	/** Used for struct, object, maps */
+	Object,
+	/** Used for array, set */
+	Array,
+	/** Used for primitive types like number, string, bool */
+	Value
+};
+
+/** Represents a custom implementation to create value in the implementation format, should be a singleton */
+struct FPropertyAnimatorCorePresetArchiveImplementation : TSharedFromThis<FPropertyAnimatorCorePresetArchiveImplementation>
+{
+	virtual ~FPropertyAnimatorCorePresetArchiveImplementation() = default;
+
+	virtual TSharedRef<FPropertyAnimatorCorePresetObjectArchive> CreateObject() = 0;
+	virtual TSharedRef<FPropertyAnimatorCorePresetArrayArchive> CreateArray() = 0;
+	virtual TSharedRef<FPropertyAnimatorCorePresetValueArchive> CreateValue(bool bInValue) = 0;
+	virtual TSharedRef<FPropertyAnimatorCorePresetValueArchive> CreateValue(uint64 InValue) = 0;
+	virtual TSharedRef<FPropertyAnimatorCorePresetValueArchive> CreateValue(int64 InValue) = 0;
+	virtual TSharedRef<FPropertyAnimatorCorePresetValueArchive> CreateValue(double InValue) = 0;
+	virtual TSharedRef<FPropertyAnimatorCorePresetValueArchive> CreateValue(const FString& InValue) = 0;
+	virtual FName GetImplementationType() const = 0;
+};
+
+/** Represents an abstract archive for preset system, underlying implementation can vary */
+struct FPropertyAnimatorCorePresetArchive : TSharedFromThis<FPropertyAnimatorCorePresetArchive>
+{
+	virtual ~FPropertyAnimatorCorePresetArchive() = default;
+
+	PROPERTYANIMATORCORE_API bool IsObject() const;
+	bool IsArray() const;
+	bool IsValue() const;
+	EPropertyAnimatorCorePresetArchiveType GetType() const;
+
+	virtual TSharedPtr<const FPropertyAnimatorCorePresetObjectArchive> AsObject() const { return nullptr; }
+	virtual TSharedPtr<const FPropertyAnimatorCorePresetArrayArchive> AsArray() const { return nullptr; }
+	virtual TSharedPtr<const FPropertyAnimatorCorePresetValueArchive> AsValue() const { return nullptr; }
+
+	virtual TSharedPtr<FPropertyAnimatorCorePresetObjectArchive> AsMutableObject() { return nullptr; }
+	virtual TSharedPtr<FPropertyAnimatorCorePresetArrayArchive> AsMutableArray() { return nullptr; }
+	virtual TSharedPtr<FPropertyAnimatorCorePresetValueArchive> AsMutableValue() { return nullptr; }
+
+	virtual bool FromString(const FString& InString) = 0;
+	virtual bool ToString(FString& OutString) const = 0;
+
+	virtual TSharedRef<FPropertyAnimatorCorePresetArchiveImplementation> GetImplementation() const = 0;
+	FName GetImplementationType() const;
+};
+
+/** Represents an abstract object archive for preset system, underlying implementation can vary */
+struct FPropertyAnimatorCorePresetObjectArchive : FPropertyAnimatorCorePresetArchive
+{
+	//~ Begin FPropertyAnimatorCorePresetArchive
+	virtual TSharedPtr<const FPropertyAnimatorCorePresetObjectArchive> AsObject() const override;
+	virtual TSharedPtr<FPropertyAnimatorCorePresetObjectArchive> AsMutableObject() override;
+	//~ End FPropertyAnimatorCorePresetArchive
+
+	bool Has(const FString& InKey, TOptional<EPropertyAnimatorCorePresetArchiveType> InType = {}) const;
+	virtual bool Remove(const FString& InKey) = 0;
+	virtual void Clear() = 0;
+
+	virtual bool Set(const FString& InKey, TSharedRef<FPropertyAnimatorCorePresetArchive> InValue) = 0;
+	virtual bool Set(const FString& InKey, bool bInValue) = 0;
+	virtual bool Set(const FString& InKey, uint64 InValue) = 0;
+	virtual bool Set(const FString& InKey, int64 InValue) = 0;
+	virtual bool Set(const FString& InKey, double InValue) = 0;
+	virtual bool Set(const FString& InKey, const FString& InValue) = 0;
+
+	virtual bool Get(const FString& InKey, TSharedPtr<FPropertyAnimatorCorePresetArchive>& OutValue) const = 0;
+	virtual bool Get(const FString& InKey, bool& bOutValue) const = 0;
+	virtual bool Get(const FString& InKey, uint64& OutValue) const = 0;
+	virtual bool Get(const FString& InKey, int64& OutValue) const = 0;
+	virtual bool Get(const FString& InKey, double& OutValue) const = 0;
+	virtual bool Get(const FString& InKey, FString& OutValue) const = 0;
+};
+
+/** Represents an abstract array archive for preset system, underlying implementation can vary */
+struct FPropertyAnimatorCorePresetArrayArchive : FPropertyAnimatorCorePresetArchive
+{
+	//~ Begin FPropertyAnimatorCorePresetArchive
+	virtual TSharedPtr<const FPropertyAnimatorCorePresetArrayArchive> AsArray() const override;
+	virtual TSharedPtr<FPropertyAnimatorCorePresetArrayArchive> AsMutableArray() override;
+	//~ End FPropertyAnimatorCorePresetArchive
+
+	virtual bool Get(int32 InIndex, TSharedPtr<FPropertyAnimatorCorePresetArchive>& OutValue) const = 0;
+	virtual int32 Num() const = 0;
+	virtual bool Remove(int32 InIndex) = 0;
+	virtual void Clear() = 0;
+
+	virtual bool Add(TSharedRef<FPropertyAnimatorCorePresetArchive> InValue) = 0;
+	virtual bool Add(bool bInValue) = 0;
+	virtual bool Add(uint64 InValue) = 0;
+	virtual bool Add(int64 InValue) = 0;
+	virtual bool Add(double InValue) = 0;
+	virtual bool Add(const FString& InValue) = 0;
+};
+
+/** Represents an abstract value (primitive) archive for preset system, underlying implementation can vary */
+struct FPropertyAnimatorCorePresetValueArchive : FPropertyAnimatorCorePresetArchive
+{
+	//~ Begin FPropertyAnimatorCorePresetArchive
+	virtual TSharedPtr<const FPropertyAnimatorCorePresetValueArchive> AsValue() const override;
+	virtual TSharedPtr<FPropertyAnimatorCorePresetValueArchive> AsMutableValue() override;
+	//~ End FPropertyAnimatorCorePresetArchive
+
+	virtual bool Get(bool& bOutValue) const = 0;
+	virtual bool Get(uint64& OutValue) const = 0;
+	virtual bool Get(int64& OutValue) const = 0;
+	virtual bool Get(double& OutValue) const = 0;
+	virtual bool Get(FString& OutValue) const = 0;
+};
diff --git a/Engine/Plugins/Experimental/PropertyAnimatorCore/Source/PropertyAnimatorCore/Public/Presets/PropertyAnimatorCorePresetBase.h b/Engine/Plugins/Experimental/PropertyAnimatorCore/Source/PropertyAnimatorCore/Public/Presets/PropertyAnimatorCorePresetBase.h
--- a/Engine/Plugins/Experimental/PropertyAnimatorCore/Source/PropertyAnimatorCore/Public/Presets/PropertyAnimatorCorePresetBase.h
+++ b/Engine/Plugins/Experimental/PropertyAnimatorCore/Source/PropertyAnimatorCore/Public/Presets/PropertyAnimatorCorePresetBase.h
@@ -2,6 +2,7 @@
 
 #pragma once
 
+#include "PropertyAnimatorCorePresetArchive.h"
 #include "PropertyAnimatorCorePresetBase.generated.h"
 
 class AActor;
@@ -61,6 +62,9 @@ public:
 	/** Called when this preset is unregistered by the subsystem */
 	virtual void OnPresetUnregistered() {}
 
+	/** Which implementation to use for archive */
+	PROPERTYANIMATORCORE_API virtual TSharedRef<FPropertyAnimatorCorePresetArchiveImplementation> GetArchiveImplementation() const;
+
 protected:
 	/** Name used to display this preset to the user */
 	UPROPERTY(VisibleInstanceOnly, Category="Animator")
@@ -70,6 +74,10 @@ protected:
 	UPROPERTY(VisibleInstanceOnly, Category="Animator")
 	int32 PresetVersion = INDEX_NONE;
 
+	/** Format used for the preset content */
+	UPROPERTY(VisibleInstanceOnly, Category="Animator")
+	FName PresetFormat;
+
 	/** Preset stringify content */
 	UPROPERTY(VisibleInstanceOnly, Category="Animator")
 	FString PresetContent;
diff --git a/Engine/Plugins/Experimental/PropertyAnimatorCore/Source/PropertyAnimatorCore/Public/Presets/PropertyAnimatorCorePresetable.h b/Engine/Plugins/Experimental/PropertyAnimatorCore/Source/PropertyAnimatorCore/Public/Presets/PropertyAnimatorCorePresetable.h
--- a/Engine/Plugins/Experimental/PropertyAnimatorCore/Source/PropertyAnimatorCore/Public/Presets/PropertyAnimatorCorePresetable.h
+++ b/Engine/Plugins/Experimental/PropertyAnimatorCore/Source/PropertyAnimatorCore/Public/Presets/PropertyAnimatorCorePresetable.h
@@ -6,8 +6,9 @@
 #include "Templates/SharedPointer.h"
 #include "PropertyAnimatorCorePresetable.generated.h"
 
-class FJsonValue;
+struct FPropertyAnimatorCorePresetArchiveImplementation;
 class UPropertyAnimatorCorePresetBase;
+struct FPropertyAnimatorCorePresetArchive;
 
 UINTERFACE(MinimalAPI)
 class UPropertyAnimatorCorePresetable : public UInterface
@@ -22,8 +23,8 @@ class IPropertyAnimatorCorePresetable
 
 public:
 	/** Import a specific preset */
-	PROPERTYANIMATORCORE_API virtual bool ImportPreset(const UPropertyAnimatorCorePresetBase* InPreset, const TSharedRef<FJsonValue>& InValue) = 0;
+	PROPERTYANIMATORCORE_API virtual bool ImportPreset(const UPropertyAnimatorCorePresetBase* InPreset, const TSharedRef<FPropertyAnimatorCorePresetArchive>& InValue) = 0;
 
 	/** Export a specific preset */
-	PROPERTYANIMATORCORE_API virtual bool ExportPreset(const UPropertyAnimatorCorePresetBase* InPreset, TSharedPtr<FJsonValue>& OutValue) = 0;
+	PROPERTYANIMATORCORE_API virtual bool ExportPreset(const UPropertyAnimatorCorePresetBase* InPreset, TSharedPtr<FPropertyAnimatorCorePresetArchive>& OutValue) const = 0;
 };
\ No newline at end of file
diff --git a/Engine/Plugins/Experimental/PropertyAnimatorCore/Source/PropertyAnimatorCore/Public/Presets/PropertyAnimatorCorePropertyPreset.h b/Engine/Plugins/Experimental/PropertyAnimatorCore/Source/PropertyAnimatorCore/Public/Presets/PropertyAnimatorCorePropertyPreset.h
--- a/Engine/Plugins/Experimental/PropertyAnimatorCore/Source/PropertyAnimatorCore/Public/Presets/PropertyAnimatorCorePropertyPreset.h
+++ b/Engine/Plugins/Experimental/PropertyAnimatorCore/Source/PropertyAnimatorCore/Public/Presets/PropertyAnimatorCorePropertyPreset.h
@@ -44,5 +44,5 @@ public:
 	PROPERTYANIMATORCORE_API virtual void GetAppliedPresetProperties(const UPropertyAnimatorCoreBase* InAnimator, TSet<FPropertyAnimatorCoreData>& OutSupportedProperties, TSet<FPropertyAnimatorCoreData>& OutAppliedProperties);
 
 protected:
-	TMap<FString, TSharedRef<FJsonValue>> PropertyPresets;
+	TMap<FString, TSharedRef<FPropertyAnimatorCorePresetArchive>> PropertyPresets;
 };
diff --git a/Engine/Plugins/Experimental/PropertyAnimatorCore/Source/PropertyAnimatorCore/Public/Properties/PropertyAnimatorCoreContext.h b/Engine/Plugins/Experimental/PropertyAnimatorCore/Source/PropertyAnimatorCore/Public/Properties/PropertyAnimatorCoreContext.h
--- a/Engine/Plugins/Experimental/PropertyAnimatorCore/Source/PropertyAnimatorCore/Public/Properties/PropertyAnimatorCoreContext.h
+++ b/Engine/Plugins/Experimental/PropertyAnimatorCore/Source/PropertyAnimatorCore/Public/Properties/PropertyAnimatorCoreContext.h
@@ -98,8 +98,8 @@ public:
 	}
 
 	//~ Begin IPropertyAnimatorCorePresetable
-	PROPERTYANIMATORCORE_API virtual bool ImportPreset(const UPropertyAnimatorCorePresetBase* InPreset, const TSharedRef<FJsonValue>& InValue) override;
-	PROPERTYANIMATORCORE_API virtual bool ExportPreset(const UPropertyAnimatorCorePresetBase* InPreset, TSharedPtr<FJsonValue>& OutValue) override;
+	PROPERTYANIMATORCORE_API virtual bool ImportPreset(const UPropertyAnimatorCorePresetBase* InPreset, const TSharedRef<FPropertyAnimatorCorePresetArchive>& InValue) override;
+	PROPERTYANIMATORCORE_API virtual bool ExportPreset(const UPropertyAnimatorCorePresetBase* InPreset, TSharedPtr<FPropertyAnimatorCorePresetArchive>& OutValue) const override;
 	//~ End IPropertyAnimatorCorePresetable
 
 protected:
diff --git a/Engine/Plugins/Experimental/PropertyAnimatorCore/Source/PropertyAnimatorCore/Public/Properties/PropertyAnimatorCoreResolver.h b/Engine/Plugins/Experimental/PropertyAnimatorCore/Source/PropertyAnimatorCore/Public/Properties/PropertyAnimatorCoreResolver.h
--- a/Engine/Plugins/Experimental/PropertyAnimatorCore/Source/PropertyAnimatorCore/Public/Properties/PropertyAnimatorCoreResolver.h
+++ b/Engine/Plugins/Experimental/PropertyAnimatorCore/Source/PropertyAnimatorCore/Public/Properties/PropertyAnimatorCoreResolver.h
@@ -39,8 +39,8 @@ public:
 	}
 
 	//~ Begin IPropertyAnimatorCorePresetable
-	PROPERTYANIMATORCORE_API virtual bool ImportPreset(const UPropertyAnimatorCorePresetBase* InPreset, const TSharedRef<FJsonValue>& InValue) override;
-	PROPERTYANIMATORCORE_API virtual bool ExportPreset(const UPropertyAnimatorCorePresetBase* InPreset, TSharedPtr<FJsonValue>& OutValue) override;
+	PROPERTYANIMATORCORE_API virtual bool ImportPreset(const UPropertyAnimatorCorePresetBase* InPreset, const TSharedRef<FPropertyAnimatorCorePresetArchive>& InValue) override;
+	PROPERTYANIMATORCORE_API virtual bool ExportPreset(const UPropertyAnimatorCorePresetBase* InPreset, TSharedPtr<FPropertyAnimatorCorePresetArchive>& OutValue) const override;
 	//~ End IPropertyAnimatorCorePresetable
 
 private:
diff --git a/Engine/Plugins/Experimental/PropertyAnimatorCore/Source/PropertyAnimatorCore/Public/TimeSources/PropertyAnimatorCoreManualTimeSource.h b/Engine/Plugins/Experimental/PropertyAnimatorCore/Source/PropertyAnimatorCore/Public/TimeSources/PropertyAnimatorCoreManualTimeSource.h
--- a/Engine/Plugins/Experimental/PropertyAnimatorCore/Source/PropertyAnimatorCore/Public/TimeSources/PropertyAnimatorCoreManualTimeSource.h
+++ b/Engine/Plugins/Experimental/PropertyAnimatorCore/Source/PropertyAnimatorCore/Public/TimeSources/PropertyAnimatorCoreManualTimeSource.h
@@ -74,8 +74,8 @@ protected:
 	virtual bool UpdateEvaluationData(FPropertyAnimatorCoreTimeSourceEvaluationData& OutData) override;
 	virtual void OnTimeSourceActive() override;
 	virtual void OnTimeSourceInactive() override;
-	virtual bool ImportPreset(const UPropertyAnimatorCorePresetBase* InPreset, const TSharedRef<FJsonValue>& InValue) override;
-	virtual bool ExportPreset(const UPropertyAnimatorCorePresetBase* InPreset, TSharedPtr<FJsonValue>& OutValue) override;
+	virtual bool ImportPreset(const UPropertyAnimatorCorePresetBase* InPreset, const TSharedRef<FPropertyAnimatorCorePresetArchive>& InValue) override;
+	virtual bool ExportPreset(const UPropertyAnimatorCorePresetBase* InPreset, TSharedPtr<FPropertyAnimatorCorePresetArchive>& OutValue) const override;
 	//~ End UPropertyAnimatorTimeSourceBase
 
 	void OnStateChanged();
diff --git a/Engine/Plugins/Experimental/PropertyAnimatorCore/Source/PropertyAnimatorCore/Public/TimeSources/PropertyAnimatorCoreSystemTimeSource.h b/Engine/Plugins/Experimental/PropertyAnimatorCore/Source/PropertyAnimatorCore/Public/TimeSources/PropertyAnimatorCoreSystemTimeSource.h
--- a/Engine/Plugins/Experimental/PropertyAnimatorCore/Source/PropertyAnimatorCore/Public/TimeSources/PropertyAnimatorCoreSystemTimeSource.h
+++ b/Engine/Plugins/Experimental/PropertyAnimatorCore/Source/PropertyAnimatorCore/Public/TimeSources/PropertyAnimatorCoreSystemTimeSource.h
@@ -62,8 +62,8 @@ protected:
 	//~ Begin UPropertyAnimatorTimeSourceBase
 	virtual bool UpdateEvaluationData(FPropertyAnimatorCoreTimeSourceEvaluationData& OutData) override;
 	virtual void OnTimeSourceActive() override;
-	virtual bool ImportPreset(const UPropertyAnimatorCorePresetBase* InPreset, const TSharedRef<FJsonValue>& InValue) override;
-	virtual bool ExportPreset(const UPropertyAnimatorCorePresetBase* InPreset, TSharedPtr<FJsonValue>& OutValue) override;
+	virtual bool ImportPreset(const UPropertyAnimatorCorePresetBase* InPreset, const TSharedRef<FPropertyAnimatorCorePresetArchive>& InValue) override;
+	virtual bool ExportPreset(const UPropertyAnimatorCorePresetBase* InPreset, TSharedPtr<FPropertyAnimatorCorePresetArchive>& OutValue) const override;
 	//~ End UPropertyAnimatorTimeSourceBase
 
 	void OnModeChanged();
diff --git a/Engine/Plugins/Experimental/PropertyAnimatorCore/Source/PropertyAnimatorCore/Public/TimeSources/PropertyAnimatorCoreTimeSourceBase.h b/Engine/Plugins/Experimental/PropertyAnimatorCore/Source/PropertyAnimatorCore/Public/TimeSources/PropertyAnimatorCoreTimeSourceBase.h
--- a/Engine/Plugins/Experimental/PropertyAnimatorCore/Source/PropertyAnimatorCore/Public/TimeSources/PropertyAnimatorCoreTimeSourceBase.h
+++ b/Engine/Plugins/Experimental/PropertyAnimatorCore/Source/PropertyAnimatorCore/Public/TimeSources/PropertyAnimatorCoreTimeSourceBase.h
@@ -81,8 +81,8 @@ public:
 	}
 
 	//~ Begin IPropertyAnimatorCorePresetable
-	PROPERTYANIMATORCORE_API virtual bool ImportPreset(const UPropertyAnimatorCorePresetBase* InPreset, const TSharedRef<FJsonValue>& InValue) override;
-	PROPERTYANIMATORCORE_API virtual bool ExportPreset(const UPropertyAnimatorCorePresetBase* InPreset, TSharedPtr<FJsonValue>& OutValue) override;
+	PROPERTYANIMATORCORE_API virtual bool ImportPreset(const UPropertyAnimatorCorePresetBase* InPreset, const TSharedRef<FPropertyAnimatorCorePresetArchive>& InValue) override;
+	PROPERTYANIMATORCORE_API virtual bool ExportPreset(const UPropertyAnimatorCorePresetBase* InPreset, TSharedPtr<FPropertyAnimatorCorePresetArchive>& OutValue) const override;
 	//~ End IPropertyAnimatorCorePresetable
 
 protected:
