MotionDesign : PropertyAnimatorCore
- Setting up pre and post roll frames for movie scene animator section to avoid restoring animated properties

diff --git a/Engine/Plugins/Experimental/PropertyAnimatorCore/Source/PropertyAnimatorCore/Private/Sequencer/MovieSceneAnimatorSection.cpp b/Engine/Plugins/Experimental/PropertyAnimatorCore/Source/PropertyAnimatorCore/Private/Sequencer/MovieSceneAnimatorSection.cpp
--- a/Engine/Plugins/Experimental/PropertyAnimatorCore/Source/PropertyAnimatorCore/Private/Sequencer/MovieSceneAnimatorSection.cpp
+++ b/Engine/Plugins/Experimental/PropertyAnimatorCore/Source/PropertyAnimatorCore/Private/Sequencer/MovieSceneAnimatorSection.cpp
@@ -5,7 +5,7 @@
 UMovieSceneAnimatorSection::UMovieSceneAnimatorSection()
 	: UMovieSceneSection()
 {
-	bSupportsInfiniteRange = true;
+	bSupportsInfiniteRange = false;
 	EvalOptions.CompletionMode = EMovieSceneCompletionMode::RestoreState;
 }
 
diff --git a/Engine/Plugins/Experimental/PropertyAnimatorCore/Source/PropertyAnimatorCore/Private/Sequencer/MovieSceneAnimatorTrack.cpp b/Engine/Plugins/Experimental/PropertyAnimatorCore/Source/PropertyAnimatorCore/Private/Sequencer/MovieSceneAnimatorTrack.cpp
--- a/Engine/Plugins/Experimental/PropertyAnimatorCore/Source/PropertyAnimatorCore/Private/Sequencer/MovieSceneAnimatorTrack.cpp
+++ b/Engine/Plugins/Experimental/PropertyAnimatorCore/Source/PropertyAnimatorCore/Private/Sequencer/MovieSceneAnimatorTrack.cpp
@@ -28,12 +28,18 @@ UMovieSceneSection* UMovieSceneAnimatorTrack::CreateNewSection()
 
 	if (const UMovieScene* MovieScene = GetTypedOuter<UMovieScene>())
 	{
+		// Playback range
 		NewSection->SetStartFrame(MovieScene->GetPlaybackRange().GetLowerBound());
 		NewSection->SetEndFrame(MovieScene->GetPlaybackRange().GetUpperBound());
 
 		// For easing
 		NewSection->SetBlendType(EMovieSceneBlendType::Absolute);
 		UpdateEasing();
+
+		// Pre/Post roll
+		const FFrameRate TickResolution = MovieScene->GetTickResolution();
+		NewSection->SetPreRollFrames((1.0 * TickResolution).RoundToFrame().Value);
+		NewSection->SetPostRollFrames((1.0 * TickResolution).RoundToFrame().Value);
 	}
 
 	return NewSection;
