MotionDesign : PropertyAnimatorCore
- Pressing Stop in manual time source will reset time back to 0

diff --git a/Engine/Plugins/Experimental/PropertyAnimatorCore/Source/PropertyAnimatorCore/Private/TimeSources/PropertyAnimatorCoreManualTimeSource.cpp b/Engine/Plugins/Experimental/PropertyAnimatorCore/Source/PropertyAnimatorCore/Private/TimeSources/PropertyAnimatorCoreManualTimeSource.cpp
--- a/Engine/Plugins/Experimental/PropertyAnimatorCore/Source/PropertyAnimatorCore/Private/TimeSources/PropertyAnimatorCoreManualTimeSource.cpp
+++ b/Engine/Plugins/Experimental/PropertyAnimatorCore/Source/PropertyAnimatorCore/Private/TimeSources/PropertyAnimatorCoreManualTimeSource.cpp
@@ -147,6 +147,7 @@ void UPropertyAnimatorCoreManualTimeSource::Stop()
 	}
 
 	Pause();
+	CustomTime = 0;
 	ActiveStatus = EPropertyAnimatorCoreManualStatus::Stopped;
 }
 
