MotionDesign : PropertyAnimatorCore
- Fix finding property owner would fail for resolvable properties like text during duplication
- Find owner object even if owner class or name do not match as long as parent class that owns the property is shared

diff --git a/Engine/Plugins/Experimental/PropertyAnimatorCore/Source/PropertyAnimatorCore/Private/Properties/PropertyAnimatorCoreContext.cpp b/Engine/Plugins/Experimental/PropertyAnimatorCore/Source/PropertyAnimatorCore/Private/Properties/PropertyAnimatorCoreContext.cpp
--- a/Engine/Plugins/Experimental/PropertyAnimatorCore/Source/PropertyAnimatorCore/Private/Properties/PropertyAnimatorCoreContext.cpp
+++ b/Engine/Plugins/Experimental/PropertyAnimatorCore/Source/PropertyAnimatorCore/Private/Properties/PropertyAnimatorCoreContext.cpp
@@ -283,91 +283,21 @@ void UPropertyAnimatorCoreContext::OnModeChanged()
 
 bool UPropertyAnimatorCoreContext::ResolvePropertyOwner(AActor* InNewOwner)
 {
-	UObject* NewOwner = InNewOwner ? InNewOwner : GetTypedOuter<AActor>();
-	const UObject* CurrentOwner = AnimatedProperty.GetOwningActor();
+	AActor* NewOwningActor = InNewOwner ? InNewOwner : GetTypedOuter<AActor>();
+	const UObject* CurrentOwningActor = AnimatedProperty.GetOwningActor();
 
-	if (CurrentOwner == NewOwner)
+	if (CurrentOwningActor == NewOwningActor)
 	{
 		return true;
 	}
 
-	const FProperty* MemberProperty = AnimatedProperty.GetMemberProperty();
-	UClass* PropertyOwningClass = MemberProperty->GetOwnerClass();
-
-	bool bFound = IsValid(NewOwner);
-	const TArray<UObject*> OtherOuters = AnimatedProperty.GetOuters(AnimatedProperty.GetOwningActor());
-
-	if (!OtherOuters.IsEmpty())
-	{
-		// Resolve using outers
-		for (const UObject* OtherOuter : OtherOuters)
-		{
-			bFound = false;
-			TArray<UObject*> ThisOwnedObjects;
-			GetObjectsWithOuter(NewOwner, ThisOwnedObjects, false);
-
-			// Search for the same class and same name
-			for (UObject* ThisOuter : ThisOwnedObjects)
-			{
-				if (ThisOuter->GetClass() == OtherOuter->GetClass()
-					&& ThisOuter->GetFName() == OtherOuter->GetFName())
-				{
-					bFound = true;
-					NewOwner = ThisOuter;
-					break;
-				}
-			}
-
-			// Search for the property owning class
-			if (!bFound)
-			{
-				for (UObject* ThisOuter : ThisOwnedObjects)
-				{
-					if (ThisOuter->GetClass() == OtherOuter->GetClass()
-						|| ThisOuter->IsA(PropertyOwningClass))
-					{
-						bFound = true;
-						NewOwner = ThisOuter;
-						break;
-					}
-				}
-
-				// Nothing found, stop searching
-				if (!bFound)
-				{
-					break;
-				}
-			}
-		}
-	}
-	else
-	{
-		// Resolve using path segments
-		const TArray<FString> OwnerSegments = AnimatedProperty.GetOuterNames();
+	const bool bFound = IsValid(NewOwningActor);
 
-		for (const FString& PathSegment : OwnerSegments)
-		{
-			bFound = false;
-			TArray<UObject*> ThisOwnedObjects;
-			GetObjectsWithOuter(NewOwner, ThisOwnedObjects, false);
-
-			for (UObject* ThisOuter : ThisOwnedObjects)
-			{
-				if (ThisOuter && ThisOuter->GetName().StartsWith(PathSegment))
-				{
-					bFound = true;
-					NewOwner = ThisOuter;
-					break;
-				}
-			}
+	// Try to resolve property owner on new owning actor
+	UObject* NewOwner = FPropertyAnimatorCoreData(NewOwningActor, AnimatedProperty.GetPropertyLocatorPath()).GetOwner();
 
-			// Nothing found, stop searching
-			if (!bFound)
-			{
-				break;
-			}
-		}
-	}
+	const FProperty* MemberProperty = AnimatedProperty.GetMemberProperty();
+	UClass* PropertyOwningClass = MemberProperty->GetOwnerClass();
 
 	if (bFound
 		&& IsValid(NewOwner)
@@ -378,7 +308,7 @@ bool UPropertyAnimatorCoreContext::ResolvePropertyOwner(AActor* InNewOwner)
 		return true;
 	}
 
-	UE_LOG(LogPropertyAnimatorCoreContext, Warning, TEXT("Could not resolve property owner %s on %s"), *AnimatedProperty.GetPathHash(), NewOwner ? *NewOwner->GetName() : TEXT("Invalid"))
+	UE_LOG(LogPropertyAnimatorCoreContext, Warning, TEXT("Could not resolve property owner %s on actor %s"), *AnimatedProperty.GetPathHash(), NewOwningActor ? *NewOwningActor->GetActorNameOrLabel() : TEXT("Invalid"))
 
 	return false;
 }
diff --git a/Engine/Plugins/Experimental/PropertyAnimatorCore/Source/PropertyAnimatorCore/Private/Properties/PropertyAnimatorCoreData.cpp b/Engine/Plugins/Experimental/PropertyAnimatorCore/Source/PropertyAnimatorCore/Private/Properties/PropertyAnimatorCoreData.cpp
--- a/Engine/Plugins/Experimental/PropertyAnimatorCore/Source/PropertyAnimatorCore/Private/Properties/PropertyAnimatorCoreData.cpp
+++ b/Engine/Plugins/Experimental/PropertyAnimatorCore/Source/PropertyAnimatorCore/Private/Properties/PropertyAnimatorCoreData.cpp
@@ -3,7 +3,6 @@
 #include "Properties/PropertyAnimatorCoreData.h"
 
 #include "Properties/PropertyAnimatorCoreResolver.h"
-#include "Properties/Handlers/PropertyAnimatorCoreHandlerBase.h"
 #include "Subsystems/PropertyAnimatorCoreSubsystem.h"
 
 FPropertyAnimatorCoreData::FPropertyAnimatorCoreData(UObject* InObject, FProperty* InMemberProperty, FProperty* InProperty, TSubclassOf<UPropertyAnimatorCoreResolver> InResolverClass)
@@ -95,7 +94,7 @@ FPropertyAnimatorCoreData::FPropertyAnimatorCoreData(AActor* InActor, const FStr
 	{
 		FString ResolverClass;
 		FString ResolverName;
-		Elements[ElementIndex].Split(TEXT(":"), &ResolverClass, &ResolverName);
+		Elements[ElementIndex++].Split(TEXT(":"), &ResolverClass, &ResolverName);
 
 		Resolver = Subsystem->FindResolverByName(FName(ResolverName));
 		Resolver = Resolver ? Resolver : Subsystem->FindResolverByClass(LoadObject<UClass>(nullptr, *ResolverClass));
@@ -104,67 +103,118 @@ FPropertyAnimatorCoreData::FPropertyAnimatorCoreData(AActor* InActor, const FStr
 		{
 			PropertyResolverClass = Resolver->GetClass();
 		}
-
-		ElementIndex++;
 	}
 
 	// Locate Owner
 	TArray<FString> Outers;
-	Elements[ElementIndex].ParseIntoArray(Outers, TEXT(","));
-
-	UObject* FoundObject = InActor;
-	for (const FString& Outer : Outers)
-	{
-		FString OuterClass;
-		FString OuterName;
-		Outer.Split(TEXT(":"), &OuterClass, &OuterName);
+	Elements[ElementIndex++].ParseIntoArray(Outers, TEXT(","));
 
-		TArray<UObject*> OwnedObjects;
-		GetObjectsWithOuter(FoundObject, OwnedObjects);
+	// Locate Property
+	TArray<FString> Properties;
+	Elements[ElementIndex].ParseIntoArray(Properties, TEXT(","));
 
-		bool bFound = false;
+	for (const FString& PropertyPath : Properties)
+	{
+		FProperty* Property = FindFProperty<FProperty>(*PropertyPath);
 
-		// Search by name
-		for (UObject* OwnedObject : OwnedObjects)
+		if (!Property)
 		{
-			if (OwnedObject->GetName() == OuterName)
-			{
-				FoundObject = OwnedObject;
-				bFound = true;
-				break;
-			}
+			break;
 		}
 
-		// Search by class
-		if (!bFound)
+		ChainProperties.Add(Property);
+	}
+
+	// Search for an object containing member property
+	if (const FProperty* MemberProperty = GetMemberProperty())
+	{
+		UClass* MemberOwningClass = MemberProperty->GetOwnerClass();
+
+		UObject* FoundObject = InActor;
+		for (int32 Index = 0; Index < Outers.Num(); Index++)
 		{
+			FString OuterClass;
+			FString OuterName;
+			Outers[Index].Split(TEXT(":"), &OuterClass, &OuterName);
+
+			TArray<UObject*> OwnedObjects;
+			GetObjectsWithOuter(FoundObject, OwnedObjects, /** IncludeNested */false);
+
+			bool bFound = false;
+
+			// Search by name
 			for (UObject* OwnedObject : OwnedObjects)
 			{
-				if (OwnedObject->GetClass()->GetClassPathName().ToString() == OuterClass)
+				if (OwnedObject->GetName() == OuterName)
 				{
 					FoundObject = OwnedObject;
 					bFound = true;
 					break;
 				}
 			}
+
+			// Search by class
+			if (!bFound)
+			{
+				for (UObject* OwnedObject : OwnedObjects)
+				{
+					if (OwnedObject->GetClass()->GetClassPathName().ToString() == OuterClass)
+					{
+						FoundObject = OwnedObject;
+						bFound = true;
+						break;
+					}
+				}
+			}
+
+			// Search by parent class when last outer to allow (USceneComponent::RelativeLocation == UStaticMeshComponent::RelativeLocation)
+			if (!bFound && (Outers.Num() - 1) == Index && MemberProperty)
+			{
+				for (UObject* OwnedObject : OwnedObjects)
+				{
+					if (OwnedObject->GetClass()->IsChildOf(MemberOwningClass))
+					{
+						FoundObject = OwnedObject;
+						bFound = true;
+						break;
+					}
+				}
+			}
+
+			if (!bFound)
+			{
+				FoundObject = nullptr;
+				break;
+			}
 		}
 
-		if (!bFound)
+		OwnerWeak = FoundObject;
+
+		// Use resolver if owner was not found
+		if (!FoundObject && Resolver && MemberProperty)
 		{
-			return;
-		}
-	}
+			const FPropertyAnimatorCoreData ActorData(InActor, nullptr, nullptr);
+			TSet<FPropertyAnimatorCoreData> ResolvableProperties;
+			Resolver->GetResolvableProperties(ActorData, ResolvableProperties);
 
-	OwnerWeak = FoundObject;
-	ElementIndex++;
+			for (const FPropertyAnimatorCoreData& ResolvableProperty : ResolvableProperties)
+			{
+				UObject* PropertyOwner = ResolvableProperty.GetOwner();
 
-	// Locate Property
-	TArray<FString> Properties;
-	Elements[ElementIndex].ParseIntoArray(Properties, TEXT(","));
+				if (!PropertyOwner)
+				{
+					continue;
+				}
 
-	for (const FString& Property : Properties)
-	{
-		ChainProperties.Add(FindFProperty<FProperty>(*Property));
+				if (PropertyOwner->IsA(MemberOwningClass))
+				{
+					FoundObject = PropertyOwner;
+					break;
+				}
+			}
+
+			OwnerWeak = FoundObject;
+		}
 	}
 
 	GeneratePropertyPath();
