MotionDesign : PropertyAnimatorCore
- Fixed linking parent property and unlinking nested properties causes a crash in properties menu

diff --git a/Engine/Plugins/Experimental/PropertyAnimatorCore/Source/PropertyAnimatorCore/Private/Animators/PropertyAnimatorCoreBase.cpp b/Engine/Plugins/Experimental/PropertyAnimatorCore/Source/PropertyAnimatorCore/Private/Animators/PropertyAnimatorCoreBase.cpp
--- a/Engine/Plugins/Experimental/PropertyAnimatorCore/Source/PropertyAnimatorCore/Private/Animators/PropertyAnimatorCoreBase.cpp
+++ b/Engine/Plugins/Experimental/PropertyAnimatorCore/Source/PropertyAnimatorCore/Private/Animators/PropertyAnimatorCoreBase.cpp
@@ -776,7 +776,7 @@ UPropertyAnimatorCoreContext* UPropertyAnimatorCoreBase::LinkProperty(const FPro
 		return PropertyContext;
 	}
 
-	if (IsPropertyLinked(InLinkProperty))
+	if (IsPropertyLinked(InLinkProperty) || !GetInnerPropertiesLinked(InLinkProperty).IsEmpty())
 	{
 		return GetLinkedPropertyContext(InLinkProperty);
 	}
@@ -807,11 +807,12 @@ bool UPropertyAnimatorCoreBase::UnlinkProperty(const FPropertyAnimatorCoreData&
 		return false;
 	}
 
-	UPropertyAnimatorCoreContext* PropertyContext = GetLinkedPropertyContext(InUnlinkProperty);
-
-	PropertyContext->Restore();
-	LinkedProperties.Remove(PropertyContext);
-	OnPropertyUnlinked(PropertyContext);
+	if (UPropertyAnimatorCoreContext* PropertyContext = GetLinkedPropertyContext(InUnlinkProperty))
+	{
+		PropertyContext->Restore();
+		LinkedProperties.Remove(PropertyContext);
+		OnPropertyUnlinked(PropertyContext);
+	}
 
 	UPropertyAnimatorCoreBase::OnAnimatorPropertyUnlinkedDelegate.Broadcast(this, InUnlinkProperty);
 
@@ -852,16 +853,11 @@ TSet<FPropertyAnimatorCoreData> UPropertyAnimatorCoreBase::GetInnerPropertiesLin
 		return OutProperties;
 	}
 
-	FProperty* LeafProperty = InPropertyData.GetLeafProperty();
-
-	for (const FPropertyAnimatorCoreData& ControllerProperty : GetLinkedProperties())
+	for (const FPropertyAnimatorCoreData& LinkedProperty : GetLinkedProperties())
 	{
-		const int32 LeafPropertyIdx = ControllerProperty.GetChainProperties().Find(LeafProperty);
-
-		// If member property is inside array and not the last one, then this controlled property is inside the InPropertyData
-		if (LeafPropertyIdx != INDEX_NONE)
+		if (InPropertyData.IsOwning(LinkedProperty))
 		{
-			OutProperties.Add(ControllerProperty);
+			OutProperties.Add(LinkedProperty);
 		}
 	}
 
diff --git a/Engine/Plugins/Experimental/PropertyAnimatorCore/Source/PropertyAnimatorCore/Private/Properties/PropertyAnimatorCoreData.cpp b/Engine/Plugins/Experimental/PropertyAnimatorCore/Source/PropertyAnimatorCore/Private/Properties/PropertyAnimatorCoreData.cpp
--- a/Engine/Plugins/Experimental/PropertyAnimatorCore/Source/PropertyAnimatorCore/Private/Properties/PropertyAnimatorCoreData.cpp
+++ b/Engine/Plugins/Experimental/PropertyAnimatorCore/Source/PropertyAnimatorCore/Private/Properties/PropertyAnimatorCoreData.cpp
@@ -466,7 +466,10 @@ bool FPropertyAnimatorCoreData::HasSetter() const
 
 bool FPropertyAnimatorCoreData::IsParentOf(const FPropertyAnimatorCoreData& InOtherProperty) const
 {
-	if (FProperty* LeafProperty = GetLeafProperty())
+	FProperty* LeafProperty = GetLeafProperty();
+	const UObject* Owner = GetOwner();
+
+	if (Owner && Owner == InOtherProperty.GetOwner() && LeafProperty)
 	{
 		const TArray<FProperty*> OtherChainProperties = InOtherProperty.GetChainProperties();
 
@@ -486,7 +489,10 @@ bool FPropertyAnimatorCoreData::IsChildOf(const FPropertyAnimatorCoreData& InOth
 
 bool FPropertyAnimatorCoreData::IsOwning(const FPropertyAnimatorCoreData& InOtherProperty) const
 {
-	if (FProperty* LeafProperty = GetLeafProperty())
+	FProperty* LeafProperty = GetLeafProperty();
+	const UObject* Owner = GetOwner();
+
+	if (Owner && Owner == InOtherProperty.GetOwner() && LeafProperty)
 	{
 		const TArray<FProperty*> OtherChainProperties = InOtherProperty.GetChainProperties();
 
diff --git a/Engine/Plugins/Experimental/PropertyAnimatorCore/Source/PropertyAnimatorCoreEditor/Private/Menus/PropertyAnimatorCoreEditorMenu.cpp b/Engine/Plugins/Experimental/PropertyAnimatorCore/Source/PropertyAnimatorCoreEditor/Private/Menus/PropertyAnimatorCoreEditorMenu.cpp
--- a/Engine/Plugins/Experimental/PropertyAnimatorCore/Source/PropertyAnimatorCoreEditor/Private/Menus/PropertyAnimatorCoreEditorMenu.cpp
+++ b/Engine/Plugins/Experimental/PropertyAnimatorCore/Source/PropertyAnimatorCoreEditor/Private/Menus/PropertyAnimatorCoreEditorMenu.cpp
@@ -834,7 +834,7 @@ void UE::PropertyAnimatorCoreEditor::Menu::FillLinkAnimatorSubmenu(UToolMenu* In
 			, FSlateIcon()
 			, FUIAction(
 				FExecuteAction::CreateLambda(&ExecuteLinkAnimatorPropertyAction, InAnimator, SupportedProperty, Empty, InMenuData)
-				, FCanExecuteAction()
+				, FCanExecuteAction::CreateLambda(&IsAnimatorLinkPropertyAllowed, InAnimator, SupportedProperty)
 				, FIsActionChecked::CreateLambda(&IsAnimatorPropertyLinked, InAnimator, SupportedProperty)
 			)
 			, EUserInterfaceActionType::ToggleButton
@@ -887,7 +887,7 @@ void UE::PropertyAnimatorCoreEditor::Menu::FillPresetAnimatorSubmenu(UToolMenu*
 			, FSlateIcon()
 			, FUIAction(
 				FExecuteAction::CreateLambda(&ExecuteLinkAnimatorPropertyAction, InAnimator, SupportedProperty, InPreset, InMenuData)
-				, FCanExecuteAction()
+				, FCanExecuteAction::CreateLambda(&IsAnimatorLinkPropertyAllowed, InAnimator, SupportedProperty)
 				, FIsActionChecked::CreateLambda(&IsAnimatorPropertyLinked, InAnimator, SupportedProperty)
 			)
 			, EUserInterfaceActionType::ToggleButton
@@ -1021,6 +1021,18 @@ bool UE::PropertyAnimatorCoreEditor::Menu::IsAnimatorPropertyLinked(const UPrope
 	return InAnimator->IsPropertyLinked(InProperty);
 }
 
+bool UE::PropertyAnimatorCoreEditor::Menu::IsAnimatorLinkPropertyAllowed(const UPropertyAnimatorCoreBase* InAnimator, const FPropertyAnimatorCoreData& InProperty)
+{
+	if (!IsValid(InAnimator) || InAnimator->IsTemplate())
+	{
+		return false;
+	}
+
+	// Only allow linking properties that are not yet linked and do not have any of their children linked
+	return (!InAnimator->IsPropertyLinked(InProperty) && InAnimator->GetInnerPropertiesLinked(InProperty).IsEmpty())
+		|| InAnimator->GetLinkedPropertyContext(InProperty) != nullptr;
+}
+
 bool UE::PropertyAnimatorCoreEditor::Menu::IsLastAnimatorCreatedPropertyLinked(const UPropertyAnimatorCoreBase* InAnimator, const FPropertyAnimatorCoreData& InProperty, TSharedRef<FPropertyAnimatorCoreEditorMenuData> InMenuData)
 {
 	if (!IsValid(InAnimator)
diff --git a/Engine/Plugins/Experimental/PropertyAnimatorCore/Source/PropertyAnimatorCoreEditor/Private/Menus/PropertyAnimatorCoreEditorMenu.h b/Engine/Plugins/Experimental/PropertyAnimatorCore/Source/PropertyAnimatorCoreEditor/Private/Menus/PropertyAnimatorCoreEditorMenu.h
--- a/Engine/Plugins/Experimental/PropertyAnimatorCore/Source/PropertyAnimatorCoreEditor/Private/Menus/PropertyAnimatorCoreEditorMenu.h
+++ b/Engine/Plugins/Experimental/PropertyAnimatorCore/Source/PropertyAnimatorCoreEditor/Private/Menus/PropertyAnimatorCoreEditorMenu.h
@@ -75,6 +75,8 @@ namespace UE::PropertyAnimatorCoreEditor::Menu
 
 	bool IsAnimatorPropertyLinked(const UPropertyAnimatorCoreBase* InAnimator, const FPropertyAnimatorCoreData& InProperty);
 
+	bool IsAnimatorLinkPropertyAllowed(const UPropertyAnimatorCoreBase* InAnimator, const FPropertyAnimatorCoreData& InProperty);
+
 	bool IsLastAnimatorCreatedPropertyLinked(const UPropertyAnimatorCoreBase* InAnimator, const FPropertyAnimatorCoreData& InProperty, TSharedRef<FPropertyAnimatorCoreEditorMenuData> InMenuData);
 
 	bool IsLastAnimatorCreatedPresetLinked(const UPropertyAnimatorCoreBase* InAnimator, UPropertyAnimatorCorePropertyPreset* InPreset, TSharedRef<FPropertyAnimatorCoreEditorMenuData> InMenuData);
