MotionDesign : PropertyAnimator
- Fix animating rotator properties would jump and cause unexpected behavior due to euler angles / gimbal lock

diff --git a/Engine/Plugins/Experimental/PropertyAnimatorCore/Source/PropertyAnimatorCore/Private/Properties/Handlers/PropertyAnimatorCoreDoubleHandler.cpp b/Engine/Plugins/Experimental/PropertyAnimatorCore/Source/PropertyAnimatorCore/Private/Properties/Handlers/PropertyAnimatorCoreDoubleHandler.cpp
--- a/Engine/Plugins/Experimental/PropertyAnimatorCore/Source/PropertyAnimatorCore/Private/Properties/Handlers/PropertyAnimatorCoreDoubleHandler.cpp
+++ b/Engine/Plugins/Experimental/PropertyAnimatorCore/Source/PropertyAnimatorCore/Private/Properties/Handlers/PropertyAnimatorCoreDoubleHandler.cpp
@@ -4,6 +4,11 @@
 
 bool UPropertyAnimatorCoreDoubleHandler::IsPropertySupported(const FPropertyAnimatorCoreData& InPropertyData) const
 {
+	if (InPropertyData.GetMemberPropertyTypeName() == NAME_Rotator)
+	{
+		return false;	
+	}
+	
 	if (InPropertyData.IsA<FDoubleProperty>())
 	{
 		return true;
diff --git a/Engine/Plugins/Experimental/PropertyAnimatorCore/Source/PropertyAnimatorCore/Private/Properties/Handlers/PropertyAnimatorCoreFloatHandler.cpp b/Engine/Plugins/Experimental/PropertyAnimatorCore/Source/PropertyAnimatorCore/Private/Properties/Handlers/PropertyAnimatorCoreFloatHandler.cpp
--- a/Engine/Plugins/Experimental/PropertyAnimatorCore/Source/PropertyAnimatorCore/Private/Properties/Handlers/PropertyAnimatorCoreFloatHandler.cpp
+++ b/Engine/Plugins/Experimental/PropertyAnimatorCore/Source/PropertyAnimatorCore/Private/Properties/Handlers/PropertyAnimatorCoreFloatHandler.cpp
@@ -4,6 +4,11 @@
 
 bool UPropertyAnimatorCoreFloatHandler::IsPropertySupported(const FPropertyAnimatorCoreData& InPropertyData) const
 {
+	if (InPropertyData.GetMemberPropertyTypeName() == NAME_Rotator)
+	{
+		return false;	
+	}
+	
 	if (InPropertyData.IsA<FFloatProperty>())
 	{
 		return true;
diff --git a/Engine/Plugins/Experimental/PropertyAnimatorCore/Source/PropertyAnimatorCore/Private/Properties/Handlers/PropertyAnimatorCoreRotatorHandler.cpp b/Engine/Plugins/Experimental/PropertyAnimatorCore/Source/PropertyAnimatorCore/Private/Properties/Handlers/PropertyAnimatorCoreRotatorHandler.cpp
--- a/Engine/Plugins/Experimental/PropertyAnimatorCore/Source/PropertyAnimatorCore/Private/Properties/Handlers/PropertyAnimatorCoreRotatorHandler.cpp
+++ b/Engine/Plugins/Experimental/PropertyAnimatorCore/Source/PropertyAnimatorCore/Private/Properties/Handlers/PropertyAnimatorCoreRotatorHandler.cpp
@@ -57,7 +57,7 @@ bool UPropertyAnimatorCoreRotatorHandler::AddValue(const FPropertyAnimatorCoreDa
 	FRotator Value;
 	InPropertyData.GetPropertyValuePtr(&Value);
 
-	FRotator NewValue = Value + *ValueResult.GetValue();
+	FRotator NewValue = FRotator(Value.Quaternion() * ValueResult.GetValue()->Quaternion());
 	InPropertyData.SetPropertyValuePtr(&NewValue);
 
 	return true;
@@ -75,7 +75,7 @@ bool UPropertyAnimatorCoreRotatorHandler::SubtractValue(const FPropertyAnimatorC
 	FRotator Value;
 	InPropertyData.GetPropertyValuePtr(&Value);
 
-	FRotator NewValue = Value - *ValueResult.GetValue();
+	FRotator NewValue = FRotator(Value.Quaternion() * ValueResult.GetValue()->Quaternion().Inverse());
 	InPropertyData.SetPropertyValuePtr(&NewValue);
 
 	return true;
