MotionDesign : PropertyAnimator
- Fix animator presets should apply and animate object by default

diff --git a/Engine/Plugins/Experimental/PropertyAnimator/Source/PropertyAnimator/Private/Presets/PropertyAnimatorPresetLocation.cpp b/Engine/Plugins/Experimental/PropertyAnimator/Source/PropertyAnimator/Private/Presets/PropertyAnimatorPresetLocation.cpp
--- a/Engine/Plugins/Experimental/PropertyAnimator/Source/PropertyAnimator/Private/Presets/PropertyAnimatorPresetLocation.cpp
+++ b/Engine/Plugins/Experimental/PropertyAnimator/Source/PropertyAnimator/Private/Presets/PropertyAnimatorPresetLocation.cpp
@@ -4,6 +4,7 @@
 
 #include "Animators/PropertyAnimatorCoreBase.h"
 #include "Properties/PropertyAnimatorFloatContext.h"
+#include "Properties/PropertyAnimatorVectorContext.h"
 
 void UPropertyAnimatorPresetLocation::GetPresetProperties(const AActor* InActor, const UPropertyAnimatorCoreBase* InAnimator, TSet<FPropertyAnimatorCoreData>& OutProperties) const
 {
@@ -25,11 +26,17 @@ void UPropertyAnimatorPresetLocation::OnPresetApplied(UPropertyAnimatorCoreBase*
 
 	for (const FPropertyAnimatorCoreData& Property : InProperties)
 	{
-		if (UPropertyAnimatorFloatContext* Context = InAnimator->GetLinkedPropertyContext<UPropertyAnimatorFloatContext>(Property))
+		if (UPropertyAnimatorFloatContext* FloatContext = InAnimator->GetLinkedPropertyContext<UPropertyAnimatorFloatContext>(Property))
 		{
-			Context->SetMode(EPropertyAnimatorCoreMode::Additive);
-			Context->SetAmplitudeMin(-100.f);
-			Context->SetAmplitudeMax(100.f);
+			FloatContext->SetMode(EPropertyAnimatorCoreMode::Additive);
+			FloatContext->SetAmplitudeMin(-100.f);
+			FloatContext->SetAmplitudeMax(100.f);
+		}
+		else if (UPropertyAnimatorVectorContext* VectorContext = InAnimator->GetLinkedPropertyContext<UPropertyAnimatorVectorContext>(Property))
+		{
+			VectorContext->SetMode(EPropertyAnimatorCoreMode::Additive);
+			VectorContext->SetAmplitudeMin(FVector(0, 0, -100.f));
+			VectorContext->SetAmplitudeMax(FVector(0, 0, 100.f));
 		}
 	}
 }
diff --git a/Engine/Plugins/Experimental/PropertyAnimator/Source/PropertyAnimator/Private/Presets/PropertyAnimatorPresetRotation.cpp b/Engine/Plugins/Experimental/PropertyAnimator/Source/PropertyAnimator/Private/Presets/PropertyAnimatorPresetRotation.cpp
--- a/Engine/Plugins/Experimental/PropertyAnimator/Source/PropertyAnimator/Private/Presets/PropertyAnimatorPresetRotation.cpp
+++ b/Engine/Plugins/Experimental/PropertyAnimator/Source/PropertyAnimator/Private/Presets/PropertyAnimatorPresetRotation.cpp
@@ -4,6 +4,7 @@
 
 #include "Animators/PropertyAnimatorCoreBase.h"
 #include "Properties/PropertyAnimatorFloatContext.h"
+#include "Properties/PropertyAnimatorRotatorContext.h"
 
 void UPropertyAnimatorPresetRotation::GetPresetProperties(const AActor* InActor, const UPropertyAnimatorCoreBase* InAnimator, TSet<FPropertyAnimatorCoreData>& OutProperties) const
 {
@@ -31,5 +32,11 @@ void UPropertyAnimatorPresetRotation::OnPresetApplied(UPropertyAnimatorCoreBase*
 			Context->SetAmplitudeMin(-90.f);
 			Context->SetAmplitudeMax(90.f);
 		}
+		else if (UPropertyAnimatorRotatorContext* RotatorContext = InAnimator->GetLinkedPropertyContext<UPropertyAnimatorRotatorContext>(Property))
+		{
+			RotatorContext->SetMode(EPropertyAnimatorCoreMode::Absolute);
+			RotatorContext->SetAmplitudeMin(FRotator(0, -90, 0));
+			RotatorContext->SetAmplitudeMax(FRotator(0, 90, 0));
+		}
 	}
 }
diff --git a/Engine/Plugins/Experimental/PropertyAnimator/Source/PropertyAnimator/Private/Presets/PropertyAnimatorPresetScale.cpp b/Engine/Plugins/Experimental/PropertyAnimator/Source/PropertyAnimator/Private/Presets/PropertyAnimatorPresetScale.cpp
--- a/Engine/Plugins/Experimental/PropertyAnimator/Source/PropertyAnimator/Private/Presets/PropertyAnimatorPresetScale.cpp
+++ b/Engine/Plugins/Experimental/PropertyAnimator/Source/PropertyAnimator/Private/Presets/PropertyAnimatorPresetScale.cpp
@@ -4,6 +4,7 @@
 
 #include "Animators/PropertyAnimatorCoreBase.h"
 #include "Properties/PropertyAnimatorFloatContext.h"
+#include "Properties/PropertyAnimatorVectorContext.h"
 
 void UPropertyAnimatorPresetScale::GetPresetProperties(const AActor* InActor, const UPropertyAnimatorCoreBase* InAnimator, TSet<FPropertyAnimatorCoreData>& OutProperties) const
 {
@@ -31,5 +32,11 @@ void UPropertyAnimatorPresetScale::OnPresetApplied(UPropertyAnimatorCoreBase* In
 			Context->SetAmplitudeMin(0.f);
 			Context->SetAmplitudeMax(1.f);
 		}
+		else if (UPropertyAnimatorVectorContext* VectorContext = InAnimator->GetLinkedPropertyContext<UPropertyAnimatorVectorContext>(Property))
+		{
+			VectorContext->SetMode(EPropertyAnimatorCoreMode::Absolute);
+			VectorContext->SetAmplitudeMin(FVector::ZeroVector);
+			VectorContext->SetAmplitudeMax(FVector::OneVector);
+		}
 	}
 }
diff --git a/Engine/Plugins/Experimental/PropertyAnimator/Source/PropertyAnimator/Private/Presets/PropertyAnimatorPresetTextLocation.cpp b/Engine/Plugins/Experimental/PropertyAnimator/Source/PropertyAnimator/Private/Presets/PropertyAnimatorPresetTextLocation.cpp
--- a/Engine/Plugins/Experimental/PropertyAnimator/Source/PropertyAnimator/Private/Presets/PropertyAnimatorPresetTextLocation.cpp
+++ b/Engine/Plugins/Experimental/PropertyAnimator/Source/PropertyAnimator/Private/Presets/PropertyAnimatorPresetTextLocation.cpp
@@ -26,19 +26,4 @@ void UPropertyAnimatorPresetTextLocation::GetPresetProperties(const AActor* InAc
 	check(LocationProperty);
 
 	OutProperties.Add(FPropertyAnimatorCoreData(TextRootComponent, LocationProperty, nullptr, UPropertyAnimatorTextResolver::StaticClass()));
-}
-
-void UPropertyAnimatorPresetTextLocation::OnPresetApplied(UPropertyAnimatorCoreBase* InAnimator, const TSet<FPropertyAnimatorCoreData>& InProperties)
-{
-	Super::OnPresetApplied(InAnimator, InProperties);
-
-	for (const FPropertyAnimatorCoreData& Property : InProperties)
-	{
-		if (UPropertyAnimatorFloatContext* Context = InAnimator->GetLinkedPropertyContext<UPropertyAnimatorFloatContext>(Property))
-		{
-			Context->SetMode(EPropertyAnimatorCoreMode::Additive);
-			Context->SetAmplitudeMin(-100.f);
-			Context->SetAmplitudeMax(100.f);
-		}
-	}
-}
+}
\ No newline at end of file
diff --git a/Engine/Plugins/Experimental/PropertyAnimator/Source/PropertyAnimator/Private/Presets/PropertyAnimatorPresetTextLocation.h b/Engine/Plugins/Experimental/PropertyAnimator/Source/PropertyAnimator/Private/Presets/PropertyAnimatorPresetTextLocation.h
--- a/Engine/Plugins/Experimental/PropertyAnimator/Source/PropertyAnimator/Private/Presets/PropertyAnimatorPresetTextLocation.h
+++ b/Engine/Plugins/Experimental/PropertyAnimator/Source/PropertyAnimator/Private/Presets/PropertyAnimatorPresetTextLocation.h
@@ -2,7 +2,7 @@
 
 #pragma once
 
-#include "Presets/PropertyAnimatorCorePropertyPreset.h"
+#include "PropertyAnimatorPresetLocation.h"
 #include "PropertyAnimatorPresetTextLocation.generated.h"
 
 class AActor;
@@ -12,7 +12,7 @@ class UPropertyAnimatorCoreBase;
  * Preset for text character position properties (X, Y, Z) on scene component
  */
 UCLASS()
-class UPropertyAnimatorPresetTextLocation : public UPropertyAnimatorCorePropertyPreset
+class UPropertyAnimatorPresetTextLocation : public UPropertyAnimatorPresetLocation
 {
 	GENERATED_BODY()
 
@@ -25,7 +25,5 @@ public:
 protected:
 	//~ Begin UPropertyAnimatorCorePresetBase
 	virtual void GetPresetProperties(const AActor* InActor, const UPropertyAnimatorCoreBase* InAnimator, TSet<FPropertyAnimatorCoreData>& OutProperties) const override;
-	virtual void OnPresetApplied(UPropertyAnimatorCoreBase* InAnimator, const TSet<FPropertyAnimatorCoreData>& InProperties) override;
-	virtual bool LoadPreset() override { return true; }
 	//~ End UPropertyAnimatorCorePresetBase
 };
\ No newline at end of file
diff --git a/Engine/Plugins/Experimental/PropertyAnimator/Source/PropertyAnimator/Private/Presets/PropertyAnimatorPresetTextRotation.cpp b/Engine/Plugins/Experimental/PropertyAnimator/Source/PropertyAnimator/Private/Presets/PropertyAnimatorPresetTextRotation.cpp
--- a/Engine/Plugins/Experimental/PropertyAnimator/Source/PropertyAnimator/Private/Presets/PropertyAnimatorPresetTextRotation.cpp
+++ b/Engine/Plugins/Experimental/PropertyAnimator/Source/PropertyAnimator/Private/Presets/PropertyAnimatorPresetTextRotation.cpp
@@ -26,19 +26,4 @@ void UPropertyAnimatorPresetTextRotation::GetPresetProperties(const AActor* InAc
 	check(RotationProperty);
 
 	OutProperties.Add(FPropertyAnimatorCoreData(TextRootComponent, RotationProperty, nullptr, UPropertyAnimatorTextResolver::StaticClass()));
-}
-
-void UPropertyAnimatorPresetTextRotation::OnPresetApplied(UPropertyAnimatorCoreBase* InAnimator, const TSet<FPropertyAnimatorCoreData>& InProperties)
-{
-	Super::OnPresetApplied(InAnimator, InProperties);
-
-	for (const FPropertyAnimatorCoreData& Property : InProperties)
-	{
-		if (UPropertyAnimatorFloatContext* Context = InAnimator->GetLinkedPropertyContext<UPropertyAnimatorFloatContext>(Property))
-		{
-			Context->SetMode(EPropertyAnimatorCoreMode::Absolute);
-			Context->SetAmplitudeMin(-90.f);
-			Context->SetAmplitudeMax(90.f);
-		}
-	}
-}
+}
\ No newline at end of file
diff --git a/Engine/Plugins/Experimental/PropertyAnimator/Source/PropertyAnimator/Private/Presets/PropertyAnimatorPresetTextRotation.h b/Engine/Plugins/Experimental/PropertyAnimator/Source/PropertyAnimator/Private/Presets/PropertyAnimatorPresetTextRotation.h
--- a/Engine/Plugins/Experimental/PropertyAnimator/Source/PropertyAnimator/Private/Presets/PropertyAnimatorPresetTextRotation.h
+++ b/Engine/Plugins/Experimental/PropertyAnimator/Source/PropertyAnimator/Private/Presets/PropertyAnimatorPresetTextRotation.h
@@ -2,7 +2,7 @@
 
 #pragma once
 
-#include "Presets/PropertyAnimatorCorePropertyPreset.h"
+#include "PropertyAnimatorPresetRotation.h"
 #include "PropertyAnimatorPresetTextRotation.generated.h"
 
 class AActor;
@@ -12,7 +12,7 @@ class UPropertyAnimatorCoreBase;
  * Preset for text character position properties (Roll, Pitch, Yaw) on scene component
  */
 UCLASS()
-class UPropertyAnimatorPresetTextRotation : public UPropertyAnimatorCorePropertyPreset
+class UPropertyAnimatorPresetTextRotation : public UPropertyAnimatorPresetRotation
 {
 	GENERATED_BODY()
 
@@ -25,7 +25,5 @@ public:
 protected:
 	//~ Begin UPropertyAnimatorCorePresetBase
 	virtual void GetPresetProperties(const AActor* InActor, const UPropertyAnimatorCoreBase* InAnimator, TSet<FPropertyAnimatorCoreData>& OutProperties) const override;
-	virtual void OnPresetApplied(UPropertyAnimatorCoreBase* InAnimator, const TSet<FPropertyAnimatorCoreData>& InProperties) override;
-	virtual bool LoadPreset() override { return true; }
 	//~ End UPropertyAnimatorCorePresetBase
 };
\ No newline at end of file
diff --git a/Engine/Plugins/Experimental/PropertyAnimator/Source/PropertyAnimator/Private/Presets/PropertyAnimatorPresetTextScale.cpp b/Engine/Plugins/Experimental/PropertyAnimator/Source/PropertyAnimator/Private/Presets/PropertyAnimatorPresetTextScale.cpp
--- a/Engine/Plugins/Experimental/PropertyAnimator/Source/PropertyAnimator/Private/Presets/PropertyAnimatorPresetTextScale.cpp
+++ b/Engine/Plugins/Experimental/PropertyAnimator/Source/PropertyAnimator/Private/Presets/PropertyAnimatorPresetTextScale.cpp
@@ -26,19 +26,4 @@ void UPropertyAnimatorPresetTextScale::GetPresetProperties(const AActor* InActor
 	check(ScaleProperty);
 
 	OutProperties.Add(FPropertyAnimatorCoreData(TextRootComponent, ScaleProperty, nullptr, UPropertyAnimatorTextResolver::StaticClass()));
-}
-
-void UPropertyAnimatorPresetTextScale::OnPresetApplied(UPropertyAnimatorCoreBase* InAnimator, const TSet<FPropertyAnimatorCoreData>& InProperties)
-{
-	Super::OnPresetApplied(InAnimator, InProperties);
-
-	for (const FPropertyAnimatorCoreData& Property : InProperties)
-	{
-		if (UPropertyAnimatorFloatContext* Context = InAnimator->GetLinkedPropertyContext<UPropertyAnimatorFloatContext>(Property))
-		{
-			Context->SetMode(EPropertyAnimatorCoreMode::Absolute);
-			Context->SetAmplitudeMin(0.f);
-			Context->SetAmplitudeMax(1.f);
-		}
-	}
-}
+}
\ No newline at end of file
diff --git a/Engine/Plugins/Experimental/PropertyAnimator/Source/PropertyAnimator/Private/Presets/PropertyAnimatorPresetTextScale.h b/Engine/Plugins/Experimental/PropertyAnimator/Source/PropertyAnimator/Private/Presets/PropertyAnimatorPresetTextScale.h
--- a/Engine/Plugins/Experimental/PropertyAnimator/Source/PropertyAnimator/Private/Presets/PropertyAnimatorPresetTextScale.h
+++ b/Engine/Plugins/Experimental/PropertyAnimator/Source/PropertyAnimator/Private/Presets/PropertyAnimatorPresetTextScale.h
@@ -2,7 +2,7 @@
 
 #pragma once
 
-#include "Presets/PropertyAnimatorCorePropertyPreset.h"
+#include "PropertyAnimatorPresetScale.h"
 #include "PropertyAnimatorPresetTextScale.generated.h"
 
 class AActor;
@@ -12,7 +12,7 @@ class UPropertyAnimatorCoreBase;
  * Preset for text character position properties (X, Y, Z) on scene component
  */
 UCLASS()
-class UPropertyAnimatorPresetTextScale : public UPropertyAnimatorCorePropertyPreset
+class UPropertyAnimatorPresetTextScale : public UPropertyAnimatorPresetScale
 {
 	GENERATED_BODY()
 
@@ -25,7 +25,5 @@ public:
 protected:
 	//~ Begin UPropertyAnimatorCorePresetBase
 	virtual void GetPresetProperties(const AActor* InActor, const UPropertyAnimatorCoreBase* InAnimator, TSet<FPropertyAnimatorCoreData>& OutProperties) const override;
-	virtual void OnPresetApplied(UPropertyAnimatorCoreBase* InAnimator, const TSet<FPropertyAnimatorCoreData>& InProperties) override;
-	virtual bool LoadPreset() override { return true; }
 	//~ End UPropertyAnimatorCorePresetBase
 };
\ No newline at end of file
diff --git a/Engine/Plugins/Experimental/PropertyAnimator/Source/PropertyAnimator/Private/Presets/PropertyAnimatorPresetTextVisibility.cpp b/Engine/Plugins/Experimental/PropertyAnimator/Source/PropertyAnimator/Private/Presets/PropertyAnimatorPresetTextVisibility.cpp
--- a/Engine/Plugins/Experimental/PropertyAnimator/Source/PropertyAnimator/Private/Presets/PropertyAnimatorPresetTextVisibility.cpp
+++ b/Engine/Plugins/Experimental/PropertyAnimator/Source/PropertyAnimator/Private/Presets/PropertyAnimatorPresetTextVisibility.cpp
@@ -26,23 +26,4 @@ void UPropertyAnimatorPresetTextVisibility::GetPresetProperties(const AActor* In
 	check(VisibilityProperty);
 
 	OutProperties.Add(FPropertyAnimatorCoreData(TextRootComponent, VisibilityProperty, nullptr, UPropertyAnimatorTextResolver::StaticClass()));
-}
-
-void UPropertyAnimatorPresetTextVisibility::OnPresetApplied(UPropertyAnimatorCoreBase* InAnimator, const TSet<FPropertyAnimatorCoreData>& InProperties)
-{
-	Super::OnPresetApplied(InAnimator, InProperties);
-
-	for (const FPropertyAnimatorCoreData& Property : InProperties)
-	{
-		if (UPropertyAnimatorCoreContext* Context = InAnimator->GetLinkedPropertyContext<UPropertyAnimatorCoreContext>(Property))
-		{
-			if (FBoolConverterRule* ConverterRule = Context->GetConverterRule<FBoolConverterRule>())
-			{
-				ConverterRule->TrueConditions.Empty();
-
-				FBoolConverterCondition NewCondition(EBoolConverterComparison::Greater, 0);
-				ConverterRule->TrueConditions.Add(NewCondition);
-			}
-		}
-	}
-}
+}
\ No newline at end of file
diff --git a/Engine/Plugins/Experimental/PropertyAnimator/Source/PropertyAnimator/Private/Presets/PropertyAnimatorPresetTextVisibility.h b/Engine/Plugins/Experimental/PropertyAnimator/Source/PropertyAnimator/Private/Presets/PropertyAnimatorPresetTextVisibility.h
--- a/Engine/Plugins/Experimental/PropertyAnimator/Source/PropertyAnimator/Private/Presets/PropertyAnimatorPresetTextVisibility.h
+++ b/Engine/Plugins/Experimental/PropertyAnimator/Source/PropertyAnimator/Private/Presets/PropertyAnimatorPresetTextVisibility.h
@@ -2,7 +2,7 @@
 
 #pragma once
 
-#include "Presets/PropertyAnimatorCorePropertyPreset.h"
+#include "PropertyAnimatorPresetVisibility.h"
 #include "PropertyAnimatorPresetTextVisibility.generated.h"
 
 class AActor;
@@ -12,7 +12,7 @@ class UPropertyAnimatorCoreBase;
  * Preset for text character visibility on scene component
  */
 UCLASS()
-class UPropertyAnimatorPresetTextVisibility : public UPropertyAnimatorCorePropertyPreset
+class UPropertyAnimatorPresetTextVisibility : public UPropertyAnimatorPresetVisibility
 {
 	GENERATED_BODY()
 
@@ -25,7 +25,5 @@ public:
 protected:
 	//~ Begin UPropertyAnimatorCorePresetBase
 	virtual void GetPresetProperties(const AActor* InActor, const UPropertyAnimatorCoreBase* InAnimator, TSet<FPropertyAnimatorCoreData>& OutProperties) const override;
-	virtual void OnPresetApplied(UPropertyAnimatorCoreBase* InAnimator, const TSet<FPropertyAnimatorCoreData>& InProperties) override;
-	virtual bool LoadPreset() override { return true; }
 	//~ End UPropertyAnimatorCorePresetBase
 };
\ No newline at end of file
diff --git a/Engine/Plugins/Experimental/PropertyAnimatorCore/Source/PropertyAnimatorCore/Private/Animators/PropertyAnimatorCoreBase.cpp b/Engine/Plugins/Experimental/PropertyAnimatorCore/Source/PropertyAnimatorCore/Private/Animators/PropertyAnimatorCoreBase.cpp
--- a/Engine/Plugins/Experimental/PropertyAnimatorCore/Source/PropertyAnimatorCore/Private/Animators/PropertyAnimatorCoreBase.cpp
+++ b/Engine/Plugins/Experimental/PropertyAnimatorCore/Source/PropertyAnimatorCore/Private/Animators/PropertyAnimatorCoreBase.cpp
@@ -4,14 +4,13 @@
 
 #include "Components/PropertyAnimatorCoreComponent.h"
 #include "GameFramework/Actor.h"
+#include "Logs/PropertyAnimatorCoreLog.h"
 #include "Properties/PropertyAnimatorCoreResolver.h"
 #include "Properties/Handlers/PropertyAnimatorCoreHandlerBase.h"
 #include "Settings/PropertyAnimatorCoreSettings.h"
 #include "Subsystems/PropertyAnimatorCoreSubsystem.h"
 #include "TimeSources/PropertyAnimatorCoreTimeSourceBase.h"
 
-DEFINE_LOG_CATEGORY_STATIC(LogPropertyAnimatorCoreBase, Log, All);
-
 UPropertyAnimatorCoreBase::FOnAnimatorUpdated UPropertyAnimatorCoreBase::OnAnimatorAddedDelegate;
 UPropertyAnimatorCoreBase::FOnAnimatorUpdated UPropertyAnimatorCoreBase::OnAnimatorRemovedDelegate;
 UPropertyAnimatorCoreBase::FOnAnimatorUpdated UPropertyAnimatorCoreBase::OnAnimatorRenamedDelegate;
@@ -477,7 +476,7 @@ void UPropertyAnimatorCoreBase::OnAnimatorRemoved()
 
 void UPropertyAnimatorCoreBase::OnAnimatorEnabled()
 {
-	UE_LOG(LogPropertyAnimatorCoreBase
+	UE_LOG(LogPropertyAnimatorCore
 		, Log
 		, TEXT("%s : PropertyAnimator %s (%s) enabled")
 		, GetAnimatorActor() ? *GetAnimatorActor()->GetActorNameOrLabel() : TEXT("Invalid Actor")
@@ -487,7 +486,7 @@ void UPropertyAnimatorCoreBase::OnAnimatorEnabled()
 
 void UPropertyAnimatorCoreBase::OnAnimatorDisabled()
 {
-	UE_LOG(LogPropertyAnimatorCoreBase
+	UE_LOG(LogPropertyAnimatorCore
 		, Log
 		, TEXT("%s : PropertyAnimator %s (%s) disabled")
 		, GetAnimatorActor() ? *GetAnimatorActor()->GetActorNameOrLabel() : TEXT("Invalid Actor")
diff --git a/Engine/Plugins/Experimental/PropertyAnimatorCore/Source/PropertyAnimatorCore/Private/Logs/PropertyAnimatorCoreLog.h b/Engine/Plugins/Experimental/PropertyAnimatorCore/Source/PropertyAnimatorCore/Private/Logs/PropertyAnimatorCoreLog.h
new file mode 100644
--- /dev/null
+++ b/Engine/Plugins/Experimental/PropertyAnimatorCore/Source/PropertyAnimatorCore/Private/Logs/PropertyAnimatorCoreLog.h
@@ -0,0 +1,7 @@
+// Copyright Epic Games, Inc. All Rights Reserved.
+
+#pragma once
+
+#include "Logging/LogMacros.h"
+
+DECLARE_LOG_CATEGORY_EXTERN(LogPropertyAnimatorCore, Log, All)
\ No newline at end of file
diff --git a/Engine/Plugins/Experimental/PropertyAnimatorCore/Source/PropertyAnimatorCore/Private/Presets/PropertyAnimatorCorePropertyPreset.cpp b/Engine/Plugins/Experimental/PropertyAnimatorCore/Source/PropertyAnimatorCore/Private/Presets/PropertyAnimatorCorePropertyPreset.cpp
--- a/Engine/Plugins/Experimental/PropertyAnimatorCore/Source/PropertyAnimatorCore/Private/Presets/PropertyAnimatorCorePropertyPreset.cpp
+++ b/Engine/Plugins/Experimental/PropertyAnimatorCore/Source/PropertyAnimatorCore/Private/Presets/PropertyAnimatorCorePropertyPreset.cpp
@@ -3,6 +3,7 @@
 #include "Presets/PropertyAnimatorCorePropertyPreset.h"
 
 #include "Animators/PropertyAnimatorCoreBase.h"
+#include "Logs/PropertyAnimatorCoreLog.h"
 #include "Properties/PropertyAnimatorCoreContext.h"
 
 void UPropertyAnimatorCorePropertyPreset::GetPresetProperties(const AActor* InActor, const UPropertyAnimatorCoreBase* InAnimator, TSet<FPropertyAnimatorCoreData>& OutProperties) const
@@ -20,6 +21,12 @@ void UPropertyAnimatorCorePropertyPreset::GetPresetProperties(const AActor* InAc
 
 void UPropertyAnimatorCorePropertyPreset::OnPresetApplied(UPropertyAnimatorCoreBase* InAnimator, const TSet<FPropertyAnimatorCoreData>& InProperties)
 {
+	if (PropertyPresets.IsEmpty())
+	{
+		UE_LOG(LogPropertyAnimatorCore, Log, TEXT("Preset %s applied on %s animator with %i properties"), *GetPresetName().ToString(), *InAnimator->GetAnimatorOriginalName().ToString(), InProperties.Num())
+		return;
+	}
+	
 	for (const TPair<FString, TSharedRef<FPropertyAnimatorCorePresetArchive>>& PropertyPreset : PropertyPresets)
 	{
 		FPropertyAnimatorCoreData Property(InAnimator->GetAnimatorActor(), PropertyPreset.Key);
@@ -48,7 +55,14 @@ void UPropertyAnimatorCorePropertyPreset::OnPresetApplied(UPropertyAnimatorCoreB
 		{
 			if (UPropertyAnimatorCoreContext* Context = InAnimator->GetLinkedPropertyContext(Property))
 			{
-				Context->ImportPreset(this, PropertyPreset.Value);
+				if (Context->ImportPreset(this, PropertyPreset.Value))
+				{
+					UE_LOG(LogPropertyAnimatorCore, Log, TEXT("Successfully imported preset %s on %s animator for property %s"), *GetPresetName().ToString(), *InAnimator->GetAnimatorOriginalName().ToString(), *Property.GetPropertyDisplayName())
+				}
+				else
+				{
+					UE_LOG(LogPropertyAnimatorCore, Warning, TEXT("Failed to import preset %s on %s animator for property %s"), *GetPresetName().ToString(), *InAnimator->GetAnimatorOriginalName().ToString(), *Property.GetPropertyDisplayName())
+				}
 			}
 		}
 	}
@@ -99,9 +113,12 @@ bool UPropertyAnimatorCorePropertyPreset::ApplyPreset(UPropertyAnimatorCoreBase*
 		return false;
 	}
 
-	for (FPropertyAnimatorCoreData& SupportedProperty : SupportedProperties)
+	for (TSet<FPropertyAnimatorCoreData>::TIterator It(SupportedProperties); It; ++It)
 	{
-		InAnimator->LinkProperty(SupportedProperty);
+		if (!InAnimator->LinkProperty(*It))
+		{
+			It.RemoveCurrent();
+		}
 	}
 
 	InAnimator->SetAnimatorDisplayName(FName(InAnimator->GetAnimatorOriginalName().ToString() + TEXT("_") + GetPresetDisplayName()));
diff --git a/Engine/Plugins/Experimental/PropertyAnimatorCore/Source/PropertyAnimatorCore/Private/Properties/PropertyAnimatorCoreContext.cpp b/Engine/Plugins/Experimental/PropertyAnimatorCore/Source/PropertyAnimatorCore/Private/Properties/PropertyAnimatorCoreContext.cpp
--- a/Engine/Plugins/Experimental/PropertyAnimatorCore/Source/PropertyAnimatorCore/Private/Properties/PropertyAnimatorCoreContext.cpp
+++ b/Engine/Plugins/Experimental/PropertyAnimatorCore/Source/PropertyAnimatorCore/Private/Properties/PropertyAnimatorCoreContext.cpp
@@ -3,14 +3,13 @@
 #include "Properties/PropertyAnimatorCoreContext.h"
 
 #include "Containers/Ticker.h"
-#include "StructUtils/InstancedStruct.h"
+#include "Logs/PropertyAnimatorCoreLog.h"
 #include "Properties/PropertyAnimatorCoreResolver.h"
 #include "Properties/Converters/PropertyAnimatorCoreConverterBase.h"
 #include "Properties/Handlers/PropertyAnimatorCoreHandlerBase.h"
+#include "StructUtils/InstancedStruct.h"
 #include "Subsystems/PropertyAnimatorCoreSubsystem.h"
 
-DEFINE_LOG_CATEGORY_STATIC(LogPropertyAnimatorCoreContext, Log, All);
-
 TArray<FPropertyAnimatorCoreData> UPropertyAnimatorCoreContext::ResolveProperty(bool bInForEvaluation) const
 {
 	TArray<FPropertyAnimatorCoreData> ResolvedProperties;
@@ -308,7 +307,7 @@ bool UPropertyAnimatorCoreContext::ResolvePropertyOwner(AActor* InNewOwner)
 		return true;
 	}
 
-	UE_LOG(LogPropertyAnimatorCoreContext, Warning, TEXT("Could not resolve property owner %s on actor %s"), *AnimatedProperty.GetPathHash(), NewOwningActor ? *NewOwningActor->GetActorNameOrLabel() : TEXT("Invalid"))
+	UE_LOG(LogPropertyAnimatorCore, Warning, TEXT("Could not resolve property owner %s on actor %s"), *AnimatedProperty.GetPathHash(), NewOwningActor ? *NewOwningActor->GetActorNameOrLabel() : TEXT("Invalid"))
 
 	return false;
 }
diff --git a/Engine/Plugins/Experimental/PropertyAnimatorCore/Source/PropertyAnimatorCore/Private/PropertyAnimatorCoreModule.cpp b/Engine/Plugins/Experimental/PropertyAnimatorCore/Source/PropertyAnimatorCore/Private/PropertyAnimatorCoreModule.cpp
--- a/Engine/Plugins/Experimental/PropertyAnimatorCore/Source/PropertyAnimatorCore/Private/PropertyAnimatorCoreModule.cpp
+++ b/Engine/Plugins/Experimental/PropertyAnimatorCore/Source/PropertyAnimatorCore/Private/PropertyAnimatorCoreModule.cpp
@@ -1,5 +1,8 @@
 // Copyright Epic Games, Inc. All Rights Reserved.
 
 #include "Modules/ModuleManager.h"
+#include "Logs/PropertyAnimatorCoreLog.h"
+
+DEFINE_LOG_CATEGORY(LogPropertyAnimatorCore)
 
 IMPLEMENT_MODULE(FDefaultModuleImpl, PropertyAnimatorCore)
\ No newline at end of file
