MotionDesign : PropertyAnimator
- Allow system time source to specify a target time instead of duration to avoid cases where time source is reactivated and countdown restarts over and operator needs to calculate the time difference

diff --git a/Engine/Plugins/Experimental/PropertyAnimatorCore/Source/PropertyAnimatorCore/Private/TimeSources/PropertyAnimatorCoreSystemTimeSource.cpp b/Engine/Plugins/Experimental/PropertyAnimatorCore/Source/PropertyAnimatorCore/Private/TimeSources/PropertyAnimatorCoreSystemTimeSource.cpp
--- a/Engine/Plugins/Experimental/PropertyAnimatorCore/Source/PropertyAnimatorCore/Private/TimeSources/PropertyAnimatorCoreSystemTimeSource.cpp
+++ b/Engine/Plugins/Experimental/PropertyAnimatorCore/Source/PropertyAnimatorCore/Private/TimeSources/PropertyAnimatorCoreSystemTimeSource.cpp
@@ -8,26 +8,31 @@
 
 bool UPropertyAnimatorCoreSystemTimeSource::UpdateEvaluationData(FPropertyAnimatorCoreTimeSourceEvaluationData& OutData)
 {
+	const FDateTime CurrentTime = bUseUtc ? FDateTime::UtcNow() : FDateTime::Now();
+
 	switch(Mode)
 	{
 	case EPropertyAnimatorCoreSystemMode::LocalTime:
 		{
-			OutData.TimeElapsed = (FDateTime::Now() - FDateTime::MinValue()).GetTotalSeconds();
-		}
-		break;
-	case EPropertyAnimatorCoreSystemMode::UtcTime:
-		{
-			OutData.TimeElapsed = (FDateTime::UtcNow() - FDateTime::MinValue()).GetTotalSeconds();
+			OutData.TimeElapsed = (CurrentTime - FDateTime::MinValue()).GetTotalSeconds();
 		}
 		break;
 	case EPropertyAnimatorCoreSystemMode::Countdown:
 		{
-			OutData.TimeElapsed = (CountdownTimeSpan - (FDateTime::Now() - ActivationTime)).GetTotalSeconds();
+			if (CountdownFormat == EPropertyAnimatorCoreSystemCountdownFormat::Target)
+			{
+				const FDateTime TargetTime = CurrentTime.GetDate() + CountdownTimeSpan;
+				OutData.TimeElapsed = (TargetTime - CurrentTime).GetTotalSeconds();
+			}
+			else
+			{
+				OutData.TimeElapsed = (CountdownTimeSpan - (CurrentTime - ActivationTime)).GetTotalSeconds();
+			}
 		}
 		break;
 	case EPropertyAnimatorCoreSystemMode::Stopwatch:
 		{
-			OutData.TimeElapsed = (FDateTime::Now() - ActivationTime).GetTotalSeconds();
+			OutData.TimeElapsed = (CurrentTime - ActivationTime).GetTotalSeconds();
 		}
 		break;
 	}
@@ -39,7 +44,7 @@ void UPropertyAnimatorCoreSystemTimeSource::OnTimeSourceActive()
 {
 	Super::OnTimeSourceActive();
 
-	ActivationTime = FDateTime::Now();
+	SetActivationTime();
 	OnModeChanged();
 }
 
@@ -49,6 +54,14 @@ bool UPropertyAnimatorCoreSystemTimeSource::ImportPreset(const UPropertyAnimator
 	{
 		const TSharedPtr<FPropertyAnimatorCorePresetObjectArchive> ObjectArchive = InValue->AsMutableObject();
 
+		bool bUseUtcValue = bUseUtc;
+		ObjectArchive->Get(GET_MEMBER_NAME_STRING_CHECKED(UPropertyAnimatorCoreSystemTimeSource, bUseUtc), bUseUtcValue);
+		SetUseUtc(bUseUtcValue);
+		
+		uint64 CountdownFormatValue = static_cast<uint64>(CountdownFormat);
+		ObjectArchive->Get(GET_MEMBER_NAME_STRING_CHECKED(UPropertyAnimatorCoreSystemTimeSource, CountdownFormat), CountdownFormatValue);
+		SetCountdownFormat(static_cast<EPropertyAnimatorCoreSystemCountdownFormat>(CountdownFormatValue));
+
 		FString CountdownValue = CountdownDuration;
 		ObjectArchive->Get(GET_MEMBER_NAME_STRING_CHECKED(UPropertyAnimatorCoreSystemTimeSource, CountdownDuration), CountdownValue);
 		SetCountdownDuration(CountdownValue);
@@ -70,6 +83,8 @@ bool UPropertyAnimatorCoreSystemTimeSource::ExportPreset(const UPropertyAnimator
 		const TSharedPtr<FPropertyAnimatorCorePresetObjectArchive> ObjectArchive = OutValue->AsMutableObject();
 
 		ObjectArchive->Set(GET_MEMBER_NAME_STRING_CHECKED(UPropertyAnimatorCoreSystemTimeSource, Mode), static_cast<uint64>(Mode));
+		ObjectArchive->Set(GET_MEMBER_NAME_STRING_CHECKED(UPropertyAnimatorCoreSystemTimeSource, bUseUtc), bUseUtc);
+		ObjectArchive->Set(GET_MEMBER_NAME_STRING_CHECKED(UPropertyAnimatorCoreSystemTimeSource, CountdownFormat), static_cast<uint64>(CountdownFormat));
 		ObjectArchive->Set(GET_MEMBER_NAME_STRING_CHECKED(UPropertyAnimatorCoreSystemTimeSource, CountdownDuration), CountdownDuration);
 
 		return true;
@@ -89,6 +104,22 @@ void UPropertyAnimatorCoreSystemTimeSource::SetMode(EPropertyAnimatorCoreSystemM
 	OnModeChanged();
 }
 
+void UPropertyAnimatorCoreSystemTimeSource::SetUseUtc(bool bInUseUtc)
+{
+	if (bUseUtc == bInUseUtc)
+	{
+		return;
+	}
+
+	bUseUtc = bInUseUtc;
+	SetActivationTime();
+}
+
+void UPropertyAnimatorCoreSystemTimeSource::SetCountdownFormat(EPropertyAnimatorCoreSystemCountdownFormat InFormat)
+{
+	CountdownFormat = InFormat;
+}
+
 void UPropertyAnimatorCoreSystemTimeSource::SetCountdownDuration(const FTimespan& InTimeSpan)
 {
 	if (InTimeSpan == CountdownTimeSpan)
@@ -122,6 +153,10 @@ void UPropertyAnimatorCoreSystemTimeSource::PostEditChangeProperty(FPropertyChan
 	{
 		OnModeChanged();
 	}
+	else if (MemberName == GET_MEMBER_NAME_CHECKED(UPropertyAnimatorCoreSystemTimeSource, bUseUtc))
+	{
+		SetActivationTime();
+	}
 }
 #endif
 
@@ -133,6 +168,11 @@ void UPropertyAnimatorCoreSystemTimeSource::OnModeChanged()
 	}
 }
 
+void UPropertyAnimatorCoreSystemTimeSource::SetActivationTime()
+{
+	ActivationTime = bUseUtc ? FDateTime::UtcNow() : FDateTime::Now();
+}
+
 FTimespan UPropertyAnimatorCoreSystemTimeSource::ParseTime(const FString& InFormat)
 {
 	// Regex patterns for different formats
diff --git a/Engine/Plugins/Experimental/PropertyAnimatorCore/Source/PropertyAnimatorCore/Public/TimeSources/PropertyAnimatorCoreSystemTimeSource.h b/Engine/Plugins/Experimental/PropertyAnimatorCore/Source/PropertyAnimatorCore/Public/TimeSources/PropertyAnimatorCoreSystemTimeSource.h
--- a/Engine/Plugins/Experimental/PropertyAnimatorCore/Source/PropertyAnimatorCore/Public/TimeSources/PropertyAnimatorCoreSystemTimeSource.h
+++ b/Engine/Plugins/Experimental/PropertyAnimatorCore/Source/PropertyAnimatorCore/Public/TimeSources/PropertyAnimatorCoreSystemTimeSource.h
@@ -13,14 +13,30 @@ enum class EPropertyAnimatorCoreSystemMode : uint8
 {
 	/** Local time of the machine */
 	LocalTime,
-	/** Universal time = Greenwich Mean Time */
-	UtcTime,
 	/** Specified duration elapsing until it reaches 0 */
 	Countdown,
 	/** Current time elapsed since the time source is active */
 	Stopwatch
 };
 
+/** Enumerates all possible format interpretation in countdown mode */
+UENUM(BlueprintType)
+enum class EPropertyAnimatorCoreSystemCountdownFormat : uint8
+{
+	/**
+	 * Default: format time provided is the duration of the countdown (relative)
+	 * eg: if current time is 10:00:00 and format is set to 12:00:00,
+	 * then countdown duration will be 12 hours until it reaches 0
+	 */
+	Duration,
+	/**
+	 * Format time provided is the target time (absolute)
+	 * eg: if current time is 10:00.00 and format is set to 12:00:00,
+	 * then countdown is 02:00:00 since 2 hours remains before reaching target time
+	 */
+	Target
+};
+
 /** System time source that support various option */
 UCLASS(MinimalAPI)
 class UPropertyAnimatorCoreSystemTimeSource : public UPropertyAnimatorCoreTimeSourceBase
@@ -38,6 +54,18 @@ public:
 		return Mode;
 	}
 
+	void SetUseUtc(bool bInUseUtc);
+	bool GetUseUtc() const
+	{
+		return bUseUtc;
+	}
+
+	void SetCountdownFormat(EPropertyAnimatorCoreSystemCountdownFormat InFormat);
+	EPropertyAnimatorCoreSystemCountdownFormat GetCountdownFormat() const
+	{
+		return CountdownFormat;
+	}
+
 	void SetCountdownDuration(const FTimespan& InTimeSpan);
 	void GetCountdownDuration(FTimespan& OutTimeSpan) const
 	{
@@ -67,13 +95,22 @@ protected:
 	//~ End UPropertyAnimatorTimeSourceBase
 
 	void OnModeChanged();
+	void SetActivationTime();
 
 	/** Machine time mode to use */
 	UPROPERTY(EditInstanceOnly, Setter, Getter, Category="Animator")
 	EPropertyAnimatorCoreSystemMode Mode = EPropertyAnimatorCoreSystemMode::LocalTime;
 
+	/** Use universal time for calculations = greenwich mean time */
+	UPROPERTY(EditInstanceOnly, Setter="SetUseUtc", Getter="GetUseUtc", Category="Animator")
+	bool bUseUtc = false;
+
+	/* Format interpretation for the countdown */
+	UPROPERTY(EditInstanceOnly, Setter, Getter, Category="Animator", meta=(EditCondition="Mode == EPropertyAnimatorCoreSystemMode::Countdown", EditConditionHides))
+	EPropertyAnimatorCoreSystemCountdownFormat CountdownFormat = EPropertyAnimatorCoreSystemCountdownFormat::Duration;
+
 	/**
-	* Countdown duration format : 
+	* Countdown time format (duration or target) : 
 	* 120 = 2 minutes
 	* 02:00 = 2 minutes
 	* 00:02:00 = 2 minutes
@@ -81,8 +118,8 @@ protected:
 	* 1h = 1 hour
 	* 120s = 2 minutes
 	*/
-	UPROPERTY(EditInstanceOnly, Setter, Getter, Category="Animator", meta=(EditCondition="Mode == EPropertyAnimatorCoreSystemMode::Countdown", EditConditionHides))
-	FString CountdownDuration = TEXT("1m");
+	UPROPERTY(EditInstanceOnly, Setter, Getter, DisplayName="Countdown Time", Category="Animator", meta=(EditCondition="Mode == EPropertyAnimatorCoreSystemMode::Countdown", EditConditionHides))
+	FString CountdownDuration = TEXT("23:59:59");
 
 private:
 	UPROPERTY(Transient)
@@ -90,4 +127,4 @@ private:
 
 	UPROPERTY(Transient)
 	FDateTime ActivationTime;
-};
\ No newline at end of file
+};
