MotionDesign : PropertyAnimator
- Ability to generate a different seed per property owner for integer property with "Seed" metadata, uses property type min/max and clamp min/max metadata for proper range

diff --git a/Engine/Plugins/Experimental/PropertyAnimator/Source/PropertyAnimator/Public/Animators/PropertyAnimatorNumericBase.h b/Engine/Plugins/Experimental/PropertyAnimator/Source/PropertyAnimator/Public/Animators/PropertyAnimatorNumericBase.h
--- a/Engine/Plugins/Experimental/PropertyAnimator/Source/PropertyAnimator/Public/Animators/PropertyAnimatorNumericBase.h
+++ b/Engine/Plugins/Experimental/PropertyAnimator/Source/PropertyAnimator/Public/Animators/PropertyAnimatorNumericBase.h
@@ -115,7 +115,7 @@ protected:
 	bool bRandomTimeOffset = false;
 
 	/** Seed to generate per property time offset */
-	UPROPERTY(EditInstanceOnly, Setter, Getter, Category="Animator", meta=(EditCondition="bRandomTimeOffset"))
+	UPROPERTY(EditInstanceOnly, Setter, Getter, Category="Animator", meta=(Seed, EditCondition="bRandomTimeOffset"))
 	int32 Seed = 0;
 
 private:
diff --git a/Engine/Plugins/Experimental/PropertyAnimatorCore/Source/PropertyAnimatorCoreEditor/Private/Customizations/PropertyAnimatorCoreEditorSeedTypeCustomization.cpp b/Engine/Plugins/Experimental/PropertyAnimatorCore/Source/PropertyAnimatorCoreEditor/Private/Customizations/PropertyAnimatorCoreEditorSeedTypeCustomization.cpp
new file mode 100644
--- /dev/null
+++ b/Engine/Plugins/Experimental/PropertyAnimatorCore/Source/PropertyAnimatorCoreEditor/Private/Customizations/PropertyAnimatorCoreEditorSeedTypeCustomization.cpp
@@ -0,0 +1,125 @@
+// Copyright Epic Games, Inc. All Rights Reserved.
+
+#include "Customizations/PropertyAnimatorCoreEditorSeedTypeCustomization.h"
+
+#include "DetailWidgetRow.h"
+#include "Input/Reply.h"
+#include "IPropertyTypeCustomization.h"
+#include "Templates/SharedPointer.h"
+#include "Widgets/Input/SButton.h"
+#include "Widgets/SBoxPanel.h"
+
+#define LOCTEXT_NAMESPACE "PropertyAnimatorCoreEditorSeedTypeCustomization"
+
+void FPropertyAnimatorCoreEditorSeedTypeCustomization::CustomizeHeader(TSharedRef<IPropertyHandle> InPropertyHandle, FDetailWidgetRow& InRow, IPropertyTypeCustomizationUtils& InUtils)
+{
+	check(InPropertyHandle->IsValidHandle() && !!InPropertyHandle->GetProperty() && InPropertyHandle->GetProperty()->IsA<FNumericProperty>())
+
+	InRow.NameContent()
+	[
+		InPropertyHandle->CreatePropertyNameWidget()
+	];
+
+	InRow.ValueContent()
+	[
+		SNew(SHorizontalBox)
+		+ SHorizontalBox::Slot()
+		.FillWidth(1.f)
+		.Padding(0.f, 0.f, 3.f, 0.f)
+		[
+			InPropertyHandle->CreatePropertyValueWidget()
+		]
+		+ SHorizontalBox::Slot()
+		.AutoWidth()
+		[
+			SNew(SButton)
+			.Text(LOCTEXT("GenerateSeed", "Generate Seed"))
+			.ToolTipText(LOCTEXT("GenerateSeedTooltip", "Generates a new seed between property type min and max value"))
+			.OnClicked_Static(&FPropertyAnimatorCoreEditorSeedTypeCustomization::OnGenerateSeedClicked, InPropertyHandle.ToWeakPtr())
+		]
+	];
+}
+
+void FPropertyAnimatorCoreEditorSeedTypeCustomization::CustomizeChildren(TSharedRef<IPropertyHandle> InPropertyHandle, IDetailChildrenBuilder& InBuilder, IPropertyTypeCustomizationUtils& InUtils)
+{
+}
+
+FReply FPropertyAnimatorCoreEditorSeedTypeCustomization::OnGenerateSeedClicked(TWeakPtr<IPropertyHandle> InPropertyHandleWeak)
+{
+	if (const TSharedPtr<IPropertyHandle> PropertyHandle = InPropertyHandleWeak.Pin())
+	{
+		const FNumericProperty* NumericProperty = CastField<FNumericProperty>(PropertyHandle->GetProperty());
+
+		if (NumericProperty && !NumericProperty->IsEnum())
+		{
+			if (NumericProperty->IsInteger())
+			{
+				int64 MinValue = 0;
+				int64 MaxValue = 0;
+
+				if (NumericProperty->IsA<FInt64Property>() || NumericProperty->IsA<FIntProperty>())
+				{
+					MinValue = TNumericLimits<int32>::Min();
+					MaxValue = TNumericLimits<int32>::Max();
+				}
+				else if (NumericProperty->IsA<FByteProperty>())
+				{
+					MinValue = TNumericLimits<uint8>::Min();
+					MaxValue = TNumericLimits<uint8>::Max();
+				}
+
+				if (PropertyHandle->HasMetaData(TEXT("ClampMin")))
+				{
+					MinValue = FMath::Max(PropertyHandle->GetIntMetaData(TEXT("ClampMin")), MinValue);
+				}
+
+				if (PropertyHandle->HasMetaData(TEXT("ClampMax")))
+				{
+					MaxValue = FMath::Min(PropertyHandle->GetIntMetaData(TEXT("ClampMax")), MaxValue);
+				}
+
+				for (int32 Index = 0; Index < PropertyHandle->GetNumPerObjectValues(); Index++)
+				{
+					const int32 Value = FMath::Lerp(MinValue, MaxValue, FMath::FRand());
+					PropertyHandle->SetPerObjectValue(Index, LexToString(Value));
+				}
+			}
+			else if (NumericProperty->IsFloatingPoint())
+			{
+				double MinValue = 0;
+				double MaxValue = 0;
+
+				if (NumericProperty->IsA<FDoubleProperty>())
+				{
+					MinValue = TNumericLimits<double>::Min();
+					MaxValue = TNumericLimits<double>::Max();
+				}
+				else if (NumericProperty->IsA<FFloatProperty>())
+				{
+					MinValue = TNumericLimits<float>::Min();
+					MaxValue = TNumericLimits<float>::Max();
+				}
+
+				if (PropertyHandle->HasMetaData(TEXT("ClampMin")))
+				{
+					MinValue = FMath::Max(PropertyHandle->GetDoubleMetaData(TEXT("ClampMin")), MinValue);
+				}
+
+				if (PropertyHandle->HasMetaData(TEXT("ClampMax")))
+				{
+					MaxValue = FMath::Min(PropertyHandle->GetDoubleMetaData(TEXT("ClampMax")), MaxValue);
+				}
+
+				for (int32 Index = 0; Index < PropertyHandle->GetNumPerObjectValues(); Index++)
+				{
+					const double Value = FMath::Lerp(MinValue, MaxValue, FMath::FRand());
+					PropertyHandle->SetPerObjectValue(Index, LexToString(Value));
+				}
+			}
+		}
+	}
+	
+	return FReply::Handled();
+}
+
+#undef LOCTEXT_NAMESPACE
diff --git a/Engine/Plugins/Experimental/PropertyAnimatorCore/Source/PropertyAnimatorCoreEditor/Private/Customizations/PropertyAnimatorCoreEditorSeedTypeCustomization.h b/Engine/Plugins/Experimental/PropertyAnimatorCore/Source/PropertyAnimatorCoreEditor/Private/Customizations/PropertyAnimatorCoreEditorSeedTypeCustomization.h
new file mode 100644
--- /dev/null
+++ b/Engine/Plugins/Experimental/PropertyAnimatorCore/Source/PropertyAnimatorCoreEditor/Private/Customizations/PropertyAnimatorCoreEditorSeedTypeCustomization.h
@@ -0,0 +1,43 @@
+// Copyright Epic Games, Inc. All Rights Reserved.
+
+#pragma once
+
+#include "IPropertyTypeCustomization.h"
+#include "PropertyEditorModule.h"
+#include "PropertyHandle.h"
+#include "Templates/SharedPointer.h"
+
+class FDetailWidgetRow;
+class FReply;
+class IDetailChildrenBuilder;
+class IPropertyTypeCustomizationUtils;
+
+/** Only allow property customization with "Seed" metadata */
+class FPropertyAnimatorCoreEditorSeedTypeIdentifier : public IPropertyTypeIdentifier
+{
+public:
+	//~ Begin IPropertyTypeIdentifier
+	virtual bool IsPropertyTypeCustomized(const IPropertyHandle& InPropertyHandle) const override
+	{
+		return InPropertyHandle.HasMetaData(TEXT("Seed"));
+	}
+	//~ End IPropertyTypeIdentifier
+};
+
+/** Type customization for seed properties */
+class FPropertyAnimatorCoreEditorSeedTypeCustomization : public IPropertyTypeCustomization
+{
+public:
+	static TSharedRef<IPropertyTypeCustomization> MakeInstance()
+	{
+		return MakeShared<FPropertyAnimatorCoreEditorSeedTypeCustomization>();
+	}
+
+	//~ Begin IPropertyTypeCustomization
+	virtual void CustomizeHeader(TSharedRef<IPropertyHandle> InPropertyHandle, FDetailWidgetRow& InRow, IPropertyTypeCustomizationUtils& InUtils) override;
+	virtual void CustomizeChildren(TSharedRef<IPropertyHandle> InPropertyHandle, IDetailChildrenBuilder& InBuilder, IPropertyTypeCustomizationUtils& InUtils) override;
+	//~ End IPropertyTypeCustomization
+
+private:
+	static FReply OnGenerateSeedClicked(TWeakPtr<IPropertyHandle> InPropertyHandleWeak);
+};
diff --git a/Engine/Plugins/Experimental/PropertyAnimatorCore/Source/PropertyAnimatorCoreEditor/Private/PropertyAnimatorCoreEditorModule.cpp b/Engine/Plugins/Experimental/PropertyAnimatorCore/Source/PropertyAnimatorCoreEditor/Private/PropertyAnimatorCoreEditorModule.cpp
--- a/Engine/Plugins/Experimental/PropertyAnimatorCore/Source/PropertyAnimatorCoreEditor/Private/PropertyAnimatorCoreEditorModule.cpp
+++ b/Engine/Plugins/Experimental/PropertyAnimatorCore/Source/PropertyAnimatorCoreEditor/Private/PropertyAnimatorCoreEditorModule.cpp
@@ -6,6 +6,7 @@
 #include "Customizations/PropertyAnimatorCoreEditorContextTypeCustomization.h"
 #include "Customizations/PropertyAnimatorCoreEditorDetailCustomization.h"
 #include "Customizations/PropertyAnimatorCoreEditorManualStateTypeCustomization.h"
+#include "Customizations/PropertyAnimatorCoreEditorSeedTypeCustomization.h"
 #include "Customizations/PropertyAnimatorCoreEditorSequencerTimeSourceEvalResultTypeCustomization.h"
 #include "ISequencerModule.h"
 #include "Modules/ModuleManager.h"
@@ -22,6 +23,7 @@ void FPropertyAnimatorCoreEditorModule::StartupModule()
 	PropertyModule.RegisterCustomPropertyTypeLayout(UPropertyAnimatorCoreContext::StaticClass()->GetFName(), FOnGetPropertyTypeCustomizationInstance::CreateStatic(&FPropertyAnimatorCoreEditorContextTypeCustomization::MakeInstance));
 	PropertyModule.RegisterCustomPropertyTypeLayout(FPropertyAnimatorCoreManualState::StaticStruct()->GetFName(), FOnGetPropertyTypeCustomizationInstance::CreateStatic(&FPropertyAnimatorCoreEditorManualStateTypeCustomization::MakeInstance));
 	PropertyModule.RegisterCustomPropertyTypeLayout(FPropertyAnimatorCoreSequencerTimeSourceEvalResult::StaticStruct()->GetFName(), FOnGetPropertyTypeCustomizationInstance::CreateStatic(&FPropertyAnimatorCoreEditorSequencerTimeSourceChannelTypeCustomization::MakeInstance));
+	PropertyModule.RegisterCustomPropertyTypeLayout(FIntProperty::StaticClass()->GetFName(), FOnGetPropertyTypeCustomizationInstance::CreateStatic(&FPropertyAnimatorCoreEditorSeedTypeCustomization::MakeInstance), MakeShared<FPropertyAnimatorCoreEditorSeedTypeIdentifier>());
 
 	PropertyModule.RegisterCustomClassLayout(UPropertyAnimatorCoreBase::StaticClass()->GetFName(), FOnGetDetailCustomizationInstance::CreateStatic(&FPropertyAnimatorCoreEditorDetailCustomization::MakeInstance));
 
@@ -31,21 +33,19 @@ void FPropertyAnimatorCoreEditorModule::StartupModule()
 
 void FPropertyAnimatorCoreEditorModule::ShutdownModule()
 {
-	if (FModuleManager::Get().IsModuleLoaded("PropertyEditor"))
+	if (FPropertyEditorModule* PropertyModule = FModuleManager::GetModulePtr<FPropertyEditorModule>("PropertyEditor"))
 	{
-		FPropertyEditorModule& PropertyModule = FModuleManager::GetModuleChecked<FPropertyEditorModule>("PropertyEditor");
+		PropertyModule->UnregisterCustomPropertyTypeLayout(UPropertyAnimatorCoreContext::StaticClass()->GetFName());
+		PropertyModule->UnregisterCustomPropertyTypeLayout(FPropertyAnimatorCoreManualState::StaticStruct()->GetFName());
+		PropertyModule->UnregisterCustomPropertyTypeLayout(FPropertyAnimatorCoreSequencerTimeSourceEvalResult::StaticStruct()->GetFName());
+		PropertyModule->UnregisterCustomPropertyTypeLayout(FIntProperty::StaticClass()->GetFName());
 
-		PropertyModule.UnregisterCustomPropertyTypeLayout(UPropertyAnimatorCoreContext::StaticClass()->GetFName());
-		PropertyModule.UnregisterCustomPropertyTypeLayout(FPropertyAnimatorCoreManualState::StaticStruct()->GetFName());
-		PropertyModule.UnregisterCustomPropertyTypeLayout(FPropertyAnimatorCoreSequencerTimeSourceEvalResult::StaticStruct()->GetFName());
-
-		PropertyModule.UnregisterCustomClassLayout(UPropertyAnimatorCoreBase::StaticClass()->GetFName());
+		PropertyModule->UnregisterCustomClassLayout(UPropertyAnimatorCoreBase::StaticClass()->GetFName());
 	}
 
-	if (FModuleManager::Get().IsModuleLoaded("Sequencer"))
+	if (ISequencerModule* SequencerModule = FModuleManager::GetModulePtr<ISequencerModule>("Sequencer"))
 	{
-		ISequencerModule& SequencerModule = FModuleManager::GetModuleChecked<ISequencerModule>("Sequencer");
-		SequencerModule.UnRegisterTrackEditor(AnimatorTrackCreateEditorHandle);
+		SequencerModule->UnRegisterTrackEditor(AnimatorTrackCreateEditorHandle);
 		AnimatorTrackCreateEditorHandle.Reset();
 	}
 }
