MotionDesign : OperatorStack
- Only focus modifier/animator tab section when a user action happened instead of editor events (load, duplicate, undo, destroy)

diff --git a/Engine/Plugins/Experimental/ActorModifierCore/Source/ActorModifierCore/Private/Modifiers/ActorModifierCoreBase.cpp b/Engine/Plugins/Experimental/ActorModifierCore/Source/ActorModifierCore/Private/Modifiers/ActorModifierCoreBase.cpp
--- a/Engine/Plugins/Experimental/ActorModifierCore/Source/ActorModifierCore/Private/Modifiers/ActorModifierCoreBase.cpp
+++ b/Engine/Plugins/Experimental/ActorModifierCore/Source/ActorModifierCore/Private/Modifiers/ActorModifierCoreBase.cpp
@@ -234,6 +234,11 @@ UActorModifierCoreStack* UActorModifierCoreBase::GetRootModifierStack() const
 	return Cast<UActorModifierCoreStack>(This);
 }
 
+UActorModifierCoreComponent* UActorModifierCoreBase::GetModifierComponent() const
+{
+	return GetTypedOuter<UActorModifierCoreComponent>();
+}
+
 const UActorModifierCoreBase* UActorModifierCoreBase::GetPreviousModifier() const
 {
 	const UActorModifierCoreBase* PreviousNameModifier = nullptr;
@@ -577,7 +582,7 @@ void UActorModifierCoreBase::PostEditUndo()
 
 	// is it an undo remove or undo add operation ?
 	const bool bModifierInStack = ModifierStack && ModifierStack->Modifiers.Contains(this);
-	const bool bStackRegistered = IsModifierStack() && !ModifierStack && ModifiedActor.IsValid();
+	const bool bStackRegistered = IsA<UActorModifierCoreStack>() && !ModifierStack;
 	const bool bModifierValid = bModifierInStack || bStackRegistered;
 
 	if (!bModifierValid)
@@ -805,7 +810,7 @@ void UActorModifierCoreBase::InitializeModifier(EActorModifierCoreEnableReason I
 
 		MarkModifierDirty(false);
 
-		UActorModifierCoreStack::OnModifierAddedDelegate.Broadcast(this);
+		UActorModifierCoreStack::OnModifierAddedDelegate.Broadcast(this, InReason);
 	}
 }
 
@@ -853,7 +858,7 @@ void UActorModifierCoreBase::UninitializeModifier(EActorModifierCoreDisableReaso
 			}
 		}
 
-		UActorModifierCoreStack::OnModifierRemovedDelegate.Broadcast(this);
+		UActorModifierCoreStack::OnModifierRemovedDelegate.Broadcast(this, InReason);
 	}
 }
 
diff --git a/Engine/Plugins/Experimental/ActorModifierCore/Source/ActorModifierCore/Private/Modifiers/ActorModifierCoreStack.cpp b/Engine/Plugins/Experimental/ActorModifierCore/Source/ActorModifierCore/Private/Modifiers/ActorModifierCoreStack.cpp
--- a/Engine/Plugins/Experimental/ActorModifierCore/Source/ActorModifierCore/Private/Modifiers/ActorModifierCoreStack.cpp
+++ b/Engine/Plugins/Experimental/ActorModifierCore/Source/ActorModifierCore/Private/Modifiers/ActorModifierCoreStack.cpp
@@ -12,8 +12,8 @@
 
 DEFINE_LOG_CATEGORY_STATIC(LogActorModifierCoreStack, Log, All);
 
-UActorModifierCoreStack::FOnModifierUpdated UActorModifierCoreStack::OnModifierAddedDelegate;
-UActorModifierCoreStack::FOnModifierUpdated UActorModifierCoreStack::OnModifierRemovedDelegate;
+UActorModifierCoreStack::FOnModifierAdded UActorModifierCoreStack::OnModifierAddedDelegate;
+UActorModifierCoreStack::FOnModifierRemoved UActorModifierCoreStack::OnModifierRemovedDelegate;
 UActorModifierCoreStack::FOnModifierUpdated UActorModifierCoreStack::OnModifierMovedDelegate;
 UActorModifierCoreStack::FOnModifierUpdated UActorModifierCoreStack::OnModifierReplacedDelegate;
 
diff --git a/Engine/Plugins/Experimental/ActorModifierCore/Source/ActorModifierCore/Public/Modifiers/ActorModifierCoreBase.h b/Engine/Plugins/Experimental/ActorModifierCore/Source/ActorModifierCore/Public/Modifiers/ActorModifierCoreBase.h
--- a/Engine/Plugins/Experimental/ActorModifierCore/Source/ActorModifierCore/Public/Modifiers/ActorModifierCoreBase.h
+++ b/Engine/Plugins/Experimental/ActorModifierCore/Source/ActorModifierCore/Public/Modifiers/ActorModifierCoreBase.h
@@ -17,6 +17,7 @@
 
 #include "ActorModifierCoreBase.generated.h"
 
+class UActorModifierCoreComponent;
 class UActorModifierCoreSharedObject;
 class UActorModifierCoreStack;
 
@@ -70,6 +71,9 @@ public:
 	/** Returns the top root stack this modifier is in */
 	ACTORMODIFIERCORE_API UActorModifierCoreStack* GetRootModifierStack() const;
 
+	/** Returns the modifier this stack is in */
+	UActorModifierCoreComponent* GetModifierComponent() const;
+
 	/** Returns the previous modifier found before this one, will recurse upwards, could return a stack */
 	const UActorModifierCoreBase* GetPreviousModifier() const;
 
diff --git a/Engine/Plugins/Experimental/ActorModifierCore/Source/ActorModifierCore/Public/Modifiers/ActorModifierCoreStack.h b/Engine/Plugins/Experimental/ActorModifierCore/Source/ActorModifierCore/Public/Modifiers/ActorModifierCoreStack.h
--- a/Engine/Plugins/Experimental/ActorModifierCore/Source/ActorModifierCore/Public/Modifiers/ActorModifierCoreStack.h
+++ b/Engine/Plugins/Experimental/ActorModifierCore/Source/ActorModifierCore/Public/Modifiers/ActorModifierCoreStack.h
@@ -26,13 +26,15 @@ class UActorModifierCoreStack : public UActorModifierCoreBase
 
 public:
 	DECLARE_MULTICAST_DELEGATE_OneParam(FOnModifierUpdated, UActorModifierCoreBase* /** UpdatedItem */)
+	DECLARE_MULTICAST_DELEGATE_TwoParams(FOnModifierAdded, UActorModifierCoreBase* /** Item */, EActorModifierCoreEnableReason /** Reason */)
+	DECLARE_MULTICAST_DELEGATE_TwoParams(FOnModifierRemoved, UActorModifierCoreBase* /** Item */, EActorModifierCoreDisableReason /** Reason */)
 
-	static FOnModifierUpdated::RegistrationType& OnModifierAdded()
+	static FOnModifierAdded::RegistrationType& OnModifierAdded()
 	{
 		return OnModifierAddedDelegate;
 	}
 
-	static FOnModifierUpdated::RegistrationType& OnModifierRemoved()
+	static FOnModifierRemoved::RegistrationType& OnModifierRemoved()
 	{
 		return OnModifierRemovedDelegate;
 	}
@@ -194,10 +196,10 @@ protected:
 
 private:
 	/** Called when a modifier is added to the stack */
-	ACTORMODIFIERCORE_API static FOnModifierUpdated OnModifierAddedDelegate;
+	ACTORMODIFIERCORE_API static FOnModifierAdded OnModifierAddedDelegate;
 
 	/** Called when a modifier is removed from the stack */
-	ACTORMODIFIERCORE_API static FOnModifierUpdated OnModifierRemovedDelegate;
+	ACTORMODIFIERCORE_API static FOnModifierRemoved OnModifierRemovedDelegate;
 
 	/** Called when a modifier is moved in the stack */
 	ACTORMODIFIERCORE_API static FOnModifierUpdated OnModifierMovedDelegate;
diff --git a/Engine/Plugins/Experimental/ActorModifierCore/Source/ActorModifierCoreEditor/Private/Modifiers/Customizations/ActorModifierCoreEditorStackCustomization.cpp b/Engine/Plugins/Experimental/ActorModifierCore/Source/ActorModifierCoreEditor/Private/Modifiers/Customizations/ActorModifierCoreEditorStackCustomization.cpp
--- a/Engine/Plugins/Experimental/ActorModifierCore/Source/ActorModifierCoreEditor/Private/Modifiers/Customizations/ActorModifierCoreEditorStackCustomization.cpp
+++ b/Engine/Plugins/Experimental/ActorModifierCore/Source/ActorModifierCoreEditor/Private/Modifiers/Customizations/ActorModifierCoreEditorStackCustomization.cpp
@@ -38,7 +38,7 @@ UActorModifierCoreEditorStackCustomization::UActorModifierCoreEditorStackCustomi
 	RegisterCustomizationFor(UActorModifierCoreBase::StaticClass());
 
 	// Modifiers delegates
-	UActorModifierCoreStack::OnModifierAdded().AddUObject(this, &UActorModifierCoreEditorStackCustomization::OnModifierUpdated);
+	UActorModifierCoreStack::OnModifierAdded().AddUObject(this, &UActorModifierCoreEditorStackCustomization::OnModifierAdded);
 	UActorModifierCoreStack::OnModifierMoved().AddUObject(this, &UActorModifierCoreEditorStackCustomization::OnModifierUpdated);
 	UActorModifierCoreStack::OnModifierRemoved().AddUObject(this, &UActorModifierCoreEditorStackCustomization::OnModifierRemoved);
 	UActorModifierCoreStack::OnModifierReplaced().AddUObject(this, &UActorModifierCoreEditorStackCustomization::OnModifierUpdated);
@@ -874,15 +874,15 @@ void UActorModifierCoreEditorStackCustomization::ToggleModifierProfilingAction(F
 	}
 }
 
-void UActorModifierCoreEditorStackCustomization::OnModifierUpdated(UActorModifierCoreBase* InModifier)
+void UActorModifierCoreEditorStackCustomization::OnModifierAdded(UActorModifierCoreBase* InModifier, EActorModifierCoreEnableReason InReason)
 {
-	if (InModifier)
+	if (InReason == EActorModifierCoreEnableReason::User)
 	{
-		RefreshActiveSelection(InModifier->GetRootModifierStack(), /** Force */false);
+		OnModifierUpdated(InModifier);
 	}
 }
 
-void UActorModifierCoreEditorStackCustomization::OnModifierRemoved(UActorModifierCoreBase* InModifier)
+void UActorModifierCoreEditorStackCustomization::OnModifierRemoved(UActorModifierCoreBase* InModifier, EActorModifierCoreDisableReason InReason)
 {
 	if (InModifier)
 	{
@@ -890,6 +890,15 @@ void UActorModifierCoreEditorStackCustomization::OnModifierRemoved(UActorModifie
 	}
 }
 
+void UActorModifierCoreEditorStackCustomization::OnModifierUpdated(UActorModifierCoreBase* InModifier)
+{
+	if (InModifier)
+	{
+		RefreshActiveSelection(InModifier->GetRootModifierStack(), /** Force */false);
+		FocusCustomization(InModifier->GetRootModifierStack());
+	}
+}
+
 TSharedRef<FUICommandList> UActorModifierCoreEditorStackCustomization::CreateModifierCommands(FOperatorStackEditorItemPtr InItem)
 {
 	TSharedRef<FUICommandList> Commands = MakeShared<FUICommandList>();
diff --git a/Engine/Plugins/Experimental/ActorModifierCore/Source/ActorModifierCoreEditor/Private/Modifiers/Customizations/ActorModifierCoreEditorStackCustomization.h b/Engine/Plugins/Experimental/ActorModifierCore/Source/ActorModifierCoreEditor/Private/Modifiers/Customizations/ActorModifierCoreEditorStackCustomization.h
--- a/Engine/Plugins/Experimental/ActorModifierCore/Source/ActorModifierCoreEditor/Private/Modifiers/Customizations/ActorModifierCoreEditorStackCustomization.h
+++ b/Engine/Plugins/Experimental/ActorModifierCore/Source/ActorModifierCoreEditor/Private/Modifiers/Customizations/ActorModifierCoreEditorStackCustomization.h
@@ -5,10 +5,12 @@
 #include "Customizations/OperatorStackEditorStackCustomization.h"
 #include "ActorModifierCoreEditorStackCustomization.generated.h"
 
-enum class EItemDropZone;
 class UActorModifierCoreStack;
 class UActorModifierCoreBase;
 class UToolMenu;
+enum class EActorModifierCoreDisableReason : uint8;
+enum class EActorModifierCoreEnableReason : uint8;
+enum class EItemDropZone;
 
 /** Struct used for copy pasting modifiers in clipboard */
 USTRUCT()
@@ -75,8 +77,9 @@ protected:
 	bool IsModifierProfiling(FOperatorStackEditorItemPtr InItem) const;
 	void ToggleModifierProfilingAction(FOperatorStackEditorItemPtr InItem) const;
 
+	void OnModifierAdded(UActorModifierCoreBase* InModifier, EActorModifierCoreEnableReason InReason);
+	void OnModifierRemoved(UActorModifierCoreBase* InModifier, EActorModifierCoreDisableReason InReason);
 	void OnModifierUpdated(UActorModifierCoreBase* InModifier);
-	void OnModifierRemoved(UActorModifierCoreBase* InModifier);
 
 	TSharedRef<FUICommandList> CreateModifierCommands(FOperatorStackEditorItemPtr InItem);
 
diff --git a/Engine/Plugins/Experimental/Avalanche/Source/AvalancheModifiersEditor/Private/Outliner/AvaOutlinerModifierProxy.cpp b/Engine/Plugins/Experimental/Avalanche/Source/AvalancheModifiersEditor/Private/Outliner/AvaOutlinerModifierProxy.cpp
--- a/Engine/Plugins/Experimental/Avalanche/Source/AvalancheModifiersEditor/Private/Outliner/AvaOutlinerModifierProxy.cpp
+++ b/Engine/Plugins/Experimental/Avalanche/Source/AvalancheModifiersEditor/Private/Outliner/AvaOutlinerModifierProxy.cpp
@@ -133,10 +133,10 @@ void FAvaOutlinerModifierProxy::GetProxiedItems(const TSharedRef<IAvaOutlinerIte
 void FAvaOutlinerModifierProxy::BindDelegates()
 {
 	UnbindDelegates();
-	UActorModifierCoreStack::OnModifierAdded().AddSP(this, &FAvaOutlinerModifierProxy::OnModifierStackUpdated);
-	UActorModifierCoreStack::OnModifierRemoved().AddSP(this, &FAvaOutlinerModifierProxy::OnModifierStackUpdated);
-	UActorModifierCoreStack::OnModifierMoved().AddSP(this, &FAvaOutlinerModifierProxy::OnModifierStackUpdated);
-	UActorModifierCoreStack::OnModifierReplaced().AddSP(this, &FAvaOutlinerModifierProxy::OnModifierStackUpdated);
+	UActorModifierCoreStack::OnModifierAdded().AddSP(this, &FAvaOutlinerModifierProxy::OnModifierAdded);
+	UActorModifierCoreStack::OnModifierRemoved().AddSP(this, &FAvaOutlinerModifierProxy::OnModifierRemoved);
+	UActorModifierCoreStack::OnModifierMoved().AddSP(this, &FAvaOutlinerModifierProxy::OnModifierUpdated);
+	UActorModifierCoreStack::OnModifierReplaced().AddSP(this, &FAvaOutlinerModifierProxy::OnModifierUpdated);
 }
 
 void FAvaOutlinerModifierProxy::UnbindDelegates()
@@ -147,7 +147,17 @@ void FAvaOutlinerModifierProxy::UnbindDelegates()
 	UActorModifierCoreStack::OnModifierReplaced().RemoveAll(this);
 }
 
-void FAvaOutlinerModifierProxy::OnModifierStackUpdated(UActorModifierCoreBase* ItemChanged)
+void FAvaOutlinerModifierProxy::OnModifierAdded(UActorModifierCoreBase* InItemChanged, EActorModifierCoreEnableReason InReason)
+{
+	OnModifierUpdated(InItemChanged);
+}
+
+void FAvaOutlinerModifierProxy::OnModifierRemoved(UActorModifierCoreBase* InItemChanged, EActorModifierCoreDisableReason InReason)
+{
+	OnModifierUpdated(InItemChanged);
+}
+
+void FAvaOutlinerModifierProxy::OnModifierUpdated(UActorModifierCoreBase* ItemChanged)
 {
 	const UActorModifierCoreStack* ThisStack = GetModifierStack();
 	const AActor* ThisActor = GetActor();
diff --git a/Engine/Plugins/Experimental/Avalanche/Source/AvalancheModifiersEditor/Private/Outliner/AvaOutlinerModifierProxy.h b/Engine/Plugins/Experimental/Avalanche/Source/AvalancheModifiersEditor/Private/Outliner/AvaOutlinerModifierProxy.h
--- a/Engine/Plugins/Experimental/Avalanche/Source/AvalancheModifiersEditor/Private/Outliner/AvaOutlinerModifierProxy.h
+++ b/Engine/Plugins/Experimental/Avalanche/Source/AvalancheModifiersEditor/Private/Outliner/AvaOutlinerModifierProxy.h
@@ -8,6 +8,8 @@
 class AActor;
 class UActorModifierCoreBase;
 class UActorModifierCoreStack;
+enum class EActorModifierCoreDisableReason : uint8;
+enum class EActorModifierCoreEnableReason : uint8;
 
 /** Creates Modifier Items based on all the Modifiers found in the Root Stack of an Actor */
 class AVALANCHEMODIFIERSEDITOR_API FAvaOutlinerModifierProxy : public FAvaOutlinerItemProxy
@@ -40,7 +42,9 @@ protected:
 	void BindDelegates();
 	void UnbindDelegates();
 
-	void OnModifierStackUpdated(UActorModifierCoreBase* ItemChanged);
+	void OnModifierAdded(UActorModifierCoreBase* InItemChanged, EActorModifierCoreEnableReason InReason);
+	void OnModifierRemoved(UActorModifierCoreBase* InItemChanged, EActorModifierCoreDisableReason InReason);
+	void OnModifierUpdated(UActorModifierCoreBase* ItemChanged);
 
 	FSlateIcon ModifierIcon;
 };
diff --git a/Engine/Plugins/Experimental/Avalanche/Source/AvalanchePropertyAnimatorEditor/Private/Outliner/AvaPropertyAnimatorEditorOutliner.cpp b/Engine/Plugins/Experimental/Avalanche/Source/AvalanchePropertyAnimatorEditor/Private/Outliner/AvaPropertyAnimatorEditorOutliner.cpp
--- a/Engine/Plugins/Experimental/Avalanche/Source/AvalanchePropertyAnimatorEditor/Private/Outliner/AvaPropertyAnimatorEditorOutliner.cpp
+++ b/Engine/Plugins/Experimental/Avalanche/Source/AvalanchePropertyAnimatorEditor/Private/Outliner/AvaPropertyAnimatorEditorOutliner.cpp
@@ -101,7 +101,7 @@ void FAvaPropertyAnimatorEditorOutliner::SetObject_Impl(UObject* InObject)
 	PropertyAnimator = Cast<UPropertyAnimatorCoreBase>(InObject);
 }
 
-void FAvaPropertyAnimatorEditorOutliner::OnAnimatorRemoved(UPropertyAnimatorCoreComponent* InComponent, UPropertyAnimatorCoreBase* InAnimator) const
+void FAvaPropertyAnimatorEditorOutliner::OnAnimatorRemoved(UPropertyAnimatorCoreComponent* InComponent, UPropertyAnimatorCoreBase* InAnimator, EPropertyAnimatorCoreUpdateEvent InReason) const
 {
 	if (InAnimator && PropertyAnimator.Get(/** EvenIfPendingKill */true) == InAnimator)
 	{
diff --git a/Engine/Plugins/Experimental/Avalanche/Source/AvalanchePropertyAnimatorEditor/Private/Outliner/AvaPropertyAnimatorEditorOutliner.h b/Engine/Plugins/Experimental/Avalanche/Source/AvalanchePropertyAnimatorEditor/Private/Outliner/AvaPropertyAnimatorEditorOutliner.h
--- a/Engine/Plugins/Experimental/Avalanche/Source/AvalanchePropertyAnimatorEditor/Private/Outliner/AvaPropertyAnimatorEditorOutliner.h
+++ b/Engine/Plugins/Experimental/Avalanche/Source/AvalanchePropertyAnimatorEditor/Private/Outliner/AvaPropertyAnimatorEditorOutliner.h
@@ -8,6 +8,7 @@
 
 class UPropertyAnimatorCoreBase;
 class UPropertyAnimatorCoreComponent;
+enum class EPropertyAnimatorCoreUpdateEvent : uint8;
 
 /** Item representing a property animator item */
 class FAvaPropertyAnimatorEditorOutliner : public FAvaOutlinerObject
@@ -40,7 +41,7 @@ protected:
 	virtual void SetObject_Impl(UObject* InObject) override;
 	//~ End FAvaOutlinerObjectItem
 
-	void OnAnimatorRemoved(UPropertyAnimatorCoreComponent* InComponent, UPropertyAnimatorCoreBase* InAnimator) const;
+	void OnAnimatorRemoved(UPropertyAnimatorCoreComponent* InComponent, UPropertyAnimatorCoreBase* InAnimator, EPropertyAnimatorCoreUpdateEvent InReason) const;
 
 	TWeakObjectPtr<UPropertyAnimatorCoreBase> PropertyAnimator;
 
diff --git a/Engine/Plugins/Experimental/Avalanche/Source/AvalanchePropertyAnimatorEditor/Private/Outliner/AvaPropertyAnimatorEditorOutlinerProxy.cpp b/Engine/Plugins/Experimental/Avalanche/Source/AvalanchePropertyAnimatorEditor/Private/Outliner/AvaPropertyAnimatorEditorOutlinerProxy.cpp
--- a/Engine/Plugins/Experimental/Avalanche/Source/AvalanchePropertyAnimatorEditor/Private/Outliner/AvaPropertyAnimatorEditorOutlinerProxy.cpp
+++ b/Engine/Plugins/Experimental/Avalanche/Source/AvalanchePropertyAnimatorEditor/Private/Outliner/AvaPropertyAnimatorEditorOutlinerProxy.cpp
@@ -138,7 +138,7 @@ void FAvaPropertyAnimatorEditorOutlinerProxy::UnbindDelegates()
 	UPropertyAnimatorCoreBase::OnPropertyAnimatorRenamed().RemoveAll(this);
 }
 
-void FAvaPropertyAnimatorEditorOutlinerProxy::OnPropertyAnimatorUpdated(UPropertyAnimatorCoreComponent* InComponent, UPropertyAnimatorCoreBase* InAnimator)
+void FAvaPropertyAnimatorEditorOutlinerProxy::OnPropertyAnimatorUpdated(UPropertyAnimatorCoreComponent* InComponent, UPropertyAnimatorCoreBase* InAnimator, EPropertyAnimatorCoreUpdateEvent InReason)
 {
 	const UPropertyAnimatorCoreComponent* ActiveAnimatorComponent = GetPropertyAnimatorComponent();
 
diff --git a/Engine/Plugins/Experimental/Avalanche/Source/AvalanchePropertyAnimatorEditor/Private/Outliner/AvaPropertyAnimatorEditorOutlinerProxy.h b/Engine/Plugins/Experimental/Avalanche/Source/AvalanchePropertyAnimatorEditor/Private/Outliner/AvaPropertyAnimatorEditorOutlinerProxy.h
--- a/Engine/Plugins/Experimental/Avalanche/Source/AvalanchePropertyAnimatorEditor/Private/Outliner/AvaPropertyAnimatorEditorOutlinerProxy.h
+++ b/Engine/Plugins/Experimental/Avalanche/Source/AvalanchePropertyAnimatorEditor/Private/Outliner/AvaPropertyAnimatorEditorOutlinerProxy.h
@@ -7,6 +7,7 @@
 
 class UPropertyAnimatorCoreBase;
 class UPropertyAnimatorCoreComponent;
+enum class EPropertyAnimatorCoreUpdateEvent : uint8;
 
 /** Creates property animator based on all the item found in the component of an Actor */
 class FAvaPropertyAnimatorEditorOutlinerProxy : public FAvaOutlinerItemProxy
@@ -37,7 +38,7 @@ protected:
 	void BindDelegates();
 	void UnbindDelegates();
 
-	void OnPropertyAnimatorUpdated(UPropertyAnimatorCoreComponent* InComponent, UPropertyAnimatorCoreBase* InAnimator);
+	void OnPropertyAnimatorUpdated(UPropertyAnimatorCoreComponent* InComponent, UPropertyAnimatorCoreBase* InAnimator, EPropertyAnimatorCoreUpdateEvent InReason);
 
 	FSlateIcon ItemIcon;
 };
diff --git a/Engine/Plugins/Experimental/OperatorStack/Source/OperatorStack/Private/Customizations/OperatorStackEditorStackCustomization.cpp b/Engine/Plugins/Experimental/OperatorStack/Source/OperatorStack/Private/Customizations/OperatorStackEditorStackCustomization.cpp
--- a/Engine/Plugins/Experimental/OperatorStack/Source/OperatorStack/Private/Customizations/OperatorStackEditorStackCustomization.cpp
+++ b/Engine/Plugins/Experimental/OperatorStack/Source/OperatorStack/Private/Customizations/OperatorStackEditorStackCustomization.cpp
@@ -90,7 +90,6 @@ void UOperatorStackEditorStackCustomization::RefreshActiveSelection(UObject* InC
 	}
 
 	Subsystem->RefreshCustomizationWidget(InContext, bInForce);
-	Subsystem->FocusCustomizationWidget(InContext, Identifier);
 }
 
 void UOperatorStackEditorStackCustomization::FocusCustomization(const UObject* InContext) const
diff --git a/Engine/Plugins/Experimental/PropertyAnimatorCore/Source/PropertyAnimatorCore/Private/Animators/PropertyAnimatorCoreBase.cpp b/Engine/Plugins/Experimental/PropertyAnimatorCore/Source/PropertyAnimatorCore/Private/Animators/PropertyAnimatorCoreBase.cpp
--- a/Engine/Plugins/Experimental/PropertyAnimatorCore/Source/PropertyAnimatorCore/Private/Animators/PropertyAnimatorCoreBase.cpp
+++ b/Engine/Plugins/Experimental/PropertyAnimatorCore/Source/PropertyAnimatorCore/Private/Animators/PropertyAnimatorCoreBase.cpp
@@ -143,7 +143,7 @@ void UPropertyAnimatorCoreBase::PostLoad()
 
 	CleanLinkedProperties();
 
-	OnAnimatorEnabledChanged();
+	OnAnimatorEnabledChanged(EPropertyAnimatorCoreUpdateEvent::Load);
 }
 
 void UPropertyAnimatorCoreBase::PostEditImport()
@@ -195,7 +195,7 @@ void UPropertyAnimatorCoreBase::PostEditChangeProperty(FPropertyChangedEvent& Pr
 
 	if (MemberName == GET_MEMBER_NAME_CHECKED(UPropertyAnimatorCoreBase, bAnimatorEnabled))
 	{
-		OnAnimatorEnabledChanged();
+		OnAnimatorEnabledChanged(EPropertyAnimatorCoreUpdateEvent::User);
 	}
 	else if (MemberName == GET_MEMBER_NAME_CHECKED(UPropertyAnimatorCoreBase, TimeSourceName)
 		|| MemberName == GET_MEMBER_NAME_CHECKED(UPropertyAnimatorCoreBase, bOverrideTimeSource))
@@ -328,7 +328,7 @@ void UPropertyAnimatorCoreBase::SetAnimatorEnabled(bool bInIsEnabled)
 	}
 
 	bAnimatorEnabled = bInIsEnabled;
-	OnAnimatorEnabledChanged();
+	OnAnimatorEnabledChanged(EPropertyAnimatorCoreUpdateEvent::User);
 }
 
 void UPropertyAnimatorCoreBase::SetOverrideTimeSource(bool bInOverride)
@@ -457,41 +457,46 @@ bool UPropertyAnimatorCoreBase::HasPropertySupport(const FPropertyAnimatorCoreDa
 	return EnumHasAnyFlags(InSupportExpected, GetPropertySupport(InPropertyData));
 }
 
-void UPropertyAnimatorCoreBase::OnAnimatorAdded()
+void UPropertyAnimatorCoreBase::OnAnimatorAdded(EPropertyAnimatorCoreUpdateEvent InType)
 {
-	bOverrideTimeSource = false;
-
-	if (const UPropertyAnimatorCoreSettings* AnimatorSettings = UPropertyAnimatorCoreSettings::Get())
+	if (InType == EPropertyAnimatorCoreUpdateEvent::User)
 	{
-		SetTimeSourceName(AnimatorSettings->GetDefaultTimeSourceName());
+		bOverrideTimeSource = false;
+
+		if (const UPropertyAnimatorCoreSettings* AnimatorSettings = UPropertyAnimatorCoreSettings::Get())
+		{
+			SetTimeSourceName(AnimatorSettings->GetDefaultTimeSourceName());
+		}
 	}
 
-	UPropertyAnimatorCoreBase::OnAnimatorAddedDelegate.Broadcast(GetAnimatorComponent(), this);
+	UPropertyAnimatorCoreBase::OnAnimatorAddedDelegate.Broadcast(GetAnimatorComponent(), this, InType);
 }
 
-void UPropertyAnimatorCoreBase::OnAnimatorRemoved()
+void UPropertyAnimatorCoreBase::OnAnimatorRemoved(EPropertyAnimatorCoreUpdateEvent InType)
 {
-	UPropertyAnimatorCoreBase::OnAnimatorRemovedDelegate.Broadcast(GetAnimatorComponent(), this);
+	UPropertyAnimatorCoreBase::OnAnimatorRemovedDelegate.Broadcast(GetAnimatorComponent(), this, InType);
 }
 
-void UPropertyAnimatorCoreBase::OnAnimatorEnabled()
+void UPropertyAnimatorCoreBase::OnAnimatorEnabled(EPropertyAnimatorCoreUpdateEvent InType)
 {
 	UE_LOG(LogPropertyAnimatorCore
 		, Log
-		, TEXT("%s : PropertyAnimator %s (%s) enabled")
+		, TEXT("%s : PropertyAnimator %s (%s) enabled %i")
 		, GetAnimatorActor() ? *GetAnimatorActor()->GetActorNameOrLabel() : TEXT("Invalid Actor")
 		, *GetAnimatorDisplayName().ToString()
-		, *GetAnimatorOriginalName().ToString());
+		, *GetAnimatorOriginalName().ToString()
+		, InType);
 }
 
-void UPropertyAnimatorCoreBase::OnAnimatorDisabled()
+void UPropertyAnimatorCoreBase::OnAnimatorDisabled(EPropertyAnimatorCoreUpdateEvent InType)
 {
 	UE_LOG(LogPropertyAnimatorCore
 		, Log
-		, TEXT("%s : PropertyAnimator %s (%s) disabled")
+		, TEXT("%s : PropertyAnimator %s (%s) disabled %i")
 		, GetAnimatorActor() ? *GetAnimatorActor()->GetActorNameOrLabel() : TEXT("Invalid Actor")
 		, *GetAnimatorDisplayName().ToString()
-		, *GetAnimatorOriginalName().ToString());
+		, *GetAnimatorOriginalName().ToString()
+		, InType);
 
 	constexpr bool bForceReset = true;
 	RestoreProperties(bForceReset);
@@ -502,17 +507,17 @@ TSubclassOf<UPropertyAnimatorCoreContext> UPropertyAnimatorCoreBase::GetProperty
 	return UPropertyAnimatorCoreContext::StaticClass();
 }
 
-void UPropertyAnimatorCoreBase::OnAnimatorEnabledChanged()
+void UPropertyAnimatorCoreBase::OnAnimatorEnabledChanged(EPropertyAnimatorCoreUpdateEvent InType)
 {
 	const UPropertyAnimatorCoreComponent* AnimatorComponent = GetAnimatorComponent();
 
 	if (bAnimatorEnabled && AnimatorComponent->ShouldAnimate())
 	{
-		OnAnimatorEnabled();
+		OnAnimatorEnabled(InType);
 	}
 	else
 	{
-		OnAnimatorDisabled();
+		OnAnimatorDisabled(InType);
 	}
 }
 
@@ -728,7 +733,7 @@ void UPropertyAnimatorCoreBase::SetAnimatorDisplayName(FName InName)
 	AnimatorDisplayName = InName;
 	OnAnimatorDisplayNameChanged();
 
-	OnAnimatorRenamedDelegate.Broadcast(GetAnimatorComponent(), this);
+	OnAnimatorRenamedDelegate.Broadcast(GetAnimatorComponent(), this, EPropertyAnimatorCoreUpdateEvent::User);
 }
 
 TSet<FPropertyAnimatorCoreData> UPropertyAnimatorCoreBase::GetLinkedProperties() const
diff --git a/Engine/Plugins/Experimental/PropertyAnimatorCore/Source/PropertyAnimatorCore/Private/Components/PropertyAnimatorCoreComponent.cpp b/Engine/Plugins/Experimental/PropertyAnimatorCore/Source/PropertyAnimatorCore/Private/Components/PropertyAnimatorCoreComponent.cpp
--- a/Engine/Plugins/Experimental/PropertyAnimatorCore/Source/PropertyAnimatorCore/Private/Components/PropertyAnimatorCoreComponent.cpp
+++ b/Engine/Plugins/Experimental/PropertyAnimatorCore/Source/PropertyAnimatorCore/Private/Components/PropertyAnimatorCoreComponent.cpp
@@ -22,7 +22,7 @@ UPropertyAnimatorCoreBase* UPropertyAnimatorCoreComponent::AddAnimator(const UCl
 		PropertyAnimatorsInternal = PropertyAnimators;
 		PropertyAnimators.Add(NewAnimator);
 
-		OnAnimatorsChanged();
+		OnAnimatorsChanged(EPropertyAnimatorCoreUpdateEvent::User);
 	}
 
 	return NewAnimator;
@@ -42,12 +42,12 @@ UPropertyAnimatorCoreBase* UPropertyAnimatorCoreComponent::CloneAnimator(UProper
 	CloneAnimator = Cast<UPropertyAnimatorCoreBase>(StaticDuplicateObjectEx(Parameters));
 
 	// Force current state
-	CloneAnimator->OnAnimatorEnabledChanged();
+	CloneAnimator->OnAnimatorEnabledChanged(EPropertyAnimatorCoreUpdateEvent::User);
 
 	PropertyAnimatorsInternal = PropertyAnimators;
 	PropertyAnimators.Add(CloneAnimator);
 
-	OnAnimatorsChanged();
+	OnAnimatorsChanged(EPropertyAnimatorCoreUpdateEvent::User);
 
 	return CloneAnimator;
 }
@@ -62,7 +62,7 @@ bool UPropertyAnimatorCoreComponent::RemoveAnimator(UPropertyAnimatorCoreBase* I
 	PropertyAnimatorsInternal = PropertyAnimators;
 	PropertyAnimators.Remove(InAnimator);
 
-	OnAnimatorsChanged();
+	OnAnimatorsChanged(EPropertyAnimatorCoreUpdateEvent::User);
 
 	return true;
 }
@@ -82,7 +82,7 @@ void UPropertyAnimatorCoreComponent::OnAnimatorsSetEnabled(const UWorld* InWorld
 	}
 }
 
-void UPropertyAnimatorCoreComponent::OnAnimatorsChanged()
+void UPropertyAnimatorCoreComponent::OnAnimatorsChanged(EPropertyAnimatorCoreUpdateEvent InType)
 {
 	const TSet<TObjectPtr<UPropertyAnimatorCoreBase>> AnimatorsSet(PropertyAnimators);
 	const TSet<TObjectPtr<UPropertyAnimatorCoreBase>> AnimatorsInternalSet(PropertyAnimatorsInternal);
@@ -100,7 +100,7 @@ void UPropertyAnimatorCoreComponent::OnAnimatorsChanged()
 #endif
 
 			RemovedAnimator->SetAnimatorEnabled(false);
-			RemovedAnimator->OnAnimatorRemoved();
+			RemovedAnimator->OnAnimatorRemoved(InType);
 		}
 	}
 	for (const TObjectPtr<UPropertyAnimatorCoreBase>& AddedAnimator : AddedAnimators)
@@ -112,15 +112,15 @@ void UPropertyAnimatorCoreComponent::OnAnimatorsChanged()
 #endif
 
 			AddedAnimator->SetAnimatorDisplayName(GetAnimatorName(AddedAnimator));
-			AddedAnimator->OnAnimatorAdded();
+			AddedAnimator->OnAnimatorAdded(InType);
 			AddedAnimator->SetAnimatorEnabled(true);
 		}
 	}
 
-	OnAnimatorsEnabledChanged();
+	OnAnimatorsEnabledChanged(InType);
 }
 
-void UPropertyAnimatorCoreComponent::OnAnimatorsEnabledChanged()
+void UPropertyAnimatorCoreComponent::OnAnimatorsEnabledChanged(EPropertyAnimatorCoreUpdateEvent InType)
 {
 	const bool bEnableAnimators = ShouldAnimate();
 
@@ -143,7 +143,7 @@ void UPropertyAnimatorCoreComponent::OnAnimatorsEnabledChanged()
 			continue;
 		}
 
-		Animator->OnAnimatorEnabledChanged();
+		Animator->OnAnimatorEnabledChanged(InType);
 	}
 
 	SetComponentTickEnabled(bEnableAnimators);
@@ -232,7 +232,22 @@ void UPropertyAnimatorCoreComponent::OnComponentCreated()
 {
 	Super::OnComponentCreated();
 
-	InitAnimatorComponent();
+	if (AActor* OwningActor = GetOwner())
+	{
+		// For spawnable templates, restore and resolve properties owner
+		constexpr bool bForceRestore = true;
+
+		for (const TObjectPtr<UPropertyAnimatorCoreBase>& Animator : PropertyAnimators)
+		{
+			if (Animator)
+			{
+				Animator->RestoreProperties(bForceRestore);
+				Animator->ResolvePropertiesOwner(OwningActor);
+			}
+		}
+	}
+
+	UPropertyAnimatorCoreBase::OnAnimatorAddedDelegate.Broadcast(this, nullptr, EPropertyAnimatorCoreUpdateEvent::User);
 }
 
 UPropertyAnimatorCoreComponent* UPropertyAnimatorCoreComponent::FindOrAdd(AActor* InActor)
@@ -297,7 +312,7 @@ void UPropertyAnimatorCoreComponent::SetAnimatorsEnabled(bool bInEnabled)
 	}
 
 	bAnimatorsEnabled = bInEnabled;
-	OnAnimatorsEnabledChanged();
+	OnAnimatorsEnabledChanged(EPropertyAnimatorCoreUpdateEvent::User);
 }
 
 void UPropertyAnimatorCoreComponent::SetAnimatorsMagnitude(float InMagnitude)
@@ -310,7 +325,7 @@ void UPropertyAnimatorCoreComponent::SetAnimatorsMagnitude(float InMagnitude)
 	}
 
 	AnimatorsMagnitude = InMagnitude;
-	OnAnimatorsEnabledChanged();
+	OnAnimatorsEnabledChanged(EPropertyAnimatorCoreUpdateEvent::User);
 }
 
 void UPropertyAnimatorCoreComponent::OnComponentDestroyed(bool bInDestroyingHierarchy)
@@ -320,9 +335,9 @@ void UPropertyAnimatorCoreComponent::OnComponentDestroyed(bool bInDestroyingHier
 	PropertyAnimatorsInternal = PropertyAnimators;
 	PropertyAnimators.Empty();
 
-	OnAnimatorsChanged();
+	OnAnimatorsChanged(EPropertyAnimatorCoreUpdateEvent::Destroyed);
 
-	UPropertyAnimatorCoreBase::OnAnimatorRemovedDelegate.Broadcast(this, nullptr);
+	UPropertyAnimatorCoreBase::OnAnimatorRemovedDelegate.Broadcast(this, nullptr, EPropertyAnimatorCoreUpdateEvent::Destroyed);
 }
 
 void UPropertyAnimatorCoreComponent::TickComponent(float InDeltaTime, ELevelTick InTickType, FActorComponentTickFunction* InTickFunction)
@@ -351,6 +366,8 @@ void UPropertyAnimatorCoreComponent::PostLoad()
 	PRAGMA_ENABLE_DEPRECATION_WARNINGS
 
 	OnTimeSourceNameChanged();
+	
+	UPropertyAnimatorCoreBase::OnAnimatorAddedDelegate.Broadcast(this, nullptr, EPropertyAnimatorCoreUpdateEvent::Load);
 }
 
 void UPropertyAnimatorCoreComponent::PostEditImport()
@@ -358,6 +375,8 @@ void UPropertyAnimatorCoreComponent::PostEditImport()
 	Super::PostEditImport();
 
 	OnTimeSourceNameChanged();
+
+	UPropertyAnimatorCoreBase::OnAnimatorAddedDelegate.Broadcast(this, nullptr, EPropertyAnimatorCoreUpdateEvent::Duplicate);
 }
 
 void UPropertyAnimatorCoreComponent::PostDuplicate(EDuplicateMode::Type InMode)
@@ -365,6 +384,8 @@ void UPropertyAnimatorCoreComponent::PostDuplicate(EDuplicateMode::Type InMode)
 	Super::PostDuplicate(InMode);
 
 	OnTimeSourceNameChanged();
+
+	UPropertyAnimatorCoreBase::OnAnimatorAddedDelegate.Broadcast(this, nullptr, EPropertyAnimatorCoreUpdateEvent::Duplicate);
 }
 
 #if WITH_EDITOR
@@ -391,10 +412,10 @@ void UPropertyAnimatorCoreComponent::PostEditUndo()
 
 	if (!bRegistered && IsValid(this))
 	{
-		UPropertyAnimatorCoreBase::OnAnimatorAddedDelegate.Broadcast(this, nullptr);
+		UPropertyAnimatorCoreBase::OnAnimatorAddedDelegate.Broadcast(this, nullptr, EPropertyAnimatorCoreUpdateEvent::Undo);
 	}
 
-	OnAnimatorsChanged();
+	OnAnimatorsChanged(EPropertyAnimatorCoreUpdateEvent::Undo);
 }
 
 void UPropertyAnimatorCoreComponent::PreEditChange(FProperty* PropertyAboutToChange)
@@ -418,11 +439,11 @@ void UPropertyAnimatorCoreComponent::PostEditChangeProperty(FPropertyChangedEven
 	if (MemberName == GET_MEMBER_NAME_CHECKED(UPropertyAnimatorCoreComponent, bAnimatorsEnabled)
 		|| MemberName == GET_MEMBER_NAME_CHECKED(UPropertyAnimatorCoreComponent, AnimatorsMagnitude))
 	{
-		OnAnimatorsEnabledChanged();
+		OnAnimatorsEnabledChanged(EPropertyAnimatorCoreUpdateEvent::User);
 	}
 	else if (MemberName == GET_MEMBER_NAME_CHECKED(UPropertyAnimatorCoreComponent, PropertyAnimators))
 	{
-		OnAnimatorsChanged();
+		OnAnimatorsChanged(EPropertyAnimatorCoreUpdateEvent::User);
 	}
 	else if (MemberName == GET_MEMBER_NAME_CHECKED(UPropertyAnimatorCoreComponent, AnimatorsTimeSourceName))
 	{
@@ -435,7 +456,7 @@ void UPropertyAnimatorCoreComponent::SetAnimators(const TArray<TObjectPtr<UPrope
 {
 	PropertyAnimatorsInternal = PropertyAnimators;
 	PropertyAnimators = InAnimators;
-	OnAnimatorsChanged();
+	OnAnimatorsChanged(EPropertyAnimatorCoreUpdateEvent::User);
 }
 
 void UPropertyAnimatorCoreComponent::SetAnimatorsTimeSourceName(FName InTimeSourceName)
@@ -536,26 +557,6 @@ bool UPropertyAnimatorCoreComponent::EvaluateAnimators()
 	return true;
 }
 
-void UPropertyAnimatorCoreComponent::InitAnimatorComponent()
-{
-	if (AActor* OwningActor = GetOwner())
-	{
-		// For spawnable templates, restore and resolve properties owner
-		constexpr bool bForceRestore = true;
-
-		for (const TObjectPtr<UPropertyAnimatorCoreBase>& Animator : PropertyAnimators)
-		{
-			if (Animator)
-			{
-				Animator->RestoreProperties(bForceRestore);
-				Animator->ResolvePropertiesOwner(OwningActor);
-			}
-		}
-	}
-
-	UPropertyAnimatorCoreBase::OnAnimatorAddedDelegate.Broadcast(this, nullptr);
-}
-
 TArray<FName> UPropertyAnimatorCoreComponent::GetTimeSourceNames() const
 {
 	TArray<FName> TimeSourceNames;
diff --git a/Engine/Plugins/Experimental/PropertyAnimatorCore/Source/PropertyAnimatorCore/Public/Animators/PropertyAnimatorCoreBase.h b/Engine/Plugins/Experimental/PropertyAnimatorCore/Source/PropertyAnimatorCore/Public/Animators/PropertyAnimatorCoreBase.h
--- a/Engine/Plugins/Experimental/PropertyAnimatorCore/Source/PropertyAnimatorCore/Public/Animators/PropertyAnimatorCoreBase.h
+++ b/Engine/Plugins/Experimental/PropertyAnimatorCore/Source/PropertyAnimatorCore/Public/Animators/PropertyAnimatorCoreBase.h
@@ -26,6 +26,15 @@ struct FPropertyAnimatorCoreMetadata
 	FName Category = TEXT("Default");
 };
 
+enum class EPropertyAnimatorCoreUpdateEvent : uint8
+{
+	User,
+	Undo,
+	Load,
+	Duplicate,
+	Destroyed
+};
+
 /** Abstract base class for any Animator, holds a set of linked properties */
 UCLASS(MinimalAPI, Abstract, EditInlineNew, AutoExpandCategories=("Animator"))
 class UPropertyAnimatorCoreBase : public UObject, public IPropertyAnimatorCorePresetable
@@ -41,7 +50,7 @@ public:
 	static constexpr const TCHAR* FrequencyParameterName = TEXT("Frequency");
 	static constexpr const TCHAR* AlphaParameterName = TEXT("Alpha");
 
-	DECLARE_MULTICAST_DELEGATE_TwoParams(FOnAnimatorUpdated, UPropertyAnimatorCoreComponent* /* InComponent */, UPropertyAnimatorCoreBase* /* InAnimator */)
+	DECLARE_MULTICAST_DELEGATE_ThreeParams(FOnAnimatorUpdated, UPropertyAnimatorCoreComponent* /* InComponent */, UPropertyAnimatorCoreBase* /* InAnimator */, EPropertyAnimatorCoreUpdateEvent /** InType */)
 	DECLARE_MULTICAST_DELEGATE_TwoParams(FOnAnimatorPropertyUpdated, UPropertyAnimatorCoreBase* /* InAnimator */, const FPropertyAnimatorCoreData& /* InProperty */)
 
 	static FOnAnimatorUpdated::RegistrationType& OnPropertyAnimatorAdded()
@@ -279,11 +288,11 @@ protected:
 
 	virtual void OnAnimatorDisplayNameChanged() {}
 
-	PROPERTYANIMATORCORE_API virtual void OnAnimatorAdded();
-	PROPERTYANIMATORCORE_API virtual void OnAnimatorRemoved();
+	PROPERTYANIMATORCORE_API virtual void OnAnimatorAdded(EPropertyAnimatorCoreUpdateEvent InType);
+	PROPERTYANIMATORCORE_API virtual void OnAnimatorRemoved(EPropertyAnimatorCoreUpdateEvent InType);
 
-	PROPERTYANIMATORCORE_API virtual void OnAnimatorEnabled();
-	PROPERTYANIMATORCORE_API virtual void OnAnimatorDisabled();
+	PROPERTYANIMATORCORE_API virtual void OnAnimatorEnabled(EPropertyAnimatorCoreUpdateEvent InType);
+	PROPERTYANIMATORCORE_API virtual void OnAnimatorDisabled(EPropertyAnimatorCoreUpdateEvent InType);
 
 	virtual void OnAnimatorRegistered(FPropertyAnimatorCoreMetadata& InMetadata) {}
 
@@ -327,7 +336,7 @@ private:
 
 	void OnPropertyGroupsChanged();
 
-	void OnAnimatorEnabledChanged();
+	void OnAnimatorEnabledChanged(EPropertyAnimatorCoreUpdateEvent InType);
 
 	void CleanLinkedProperties();
 
diff --git a/Engine/Plugins/Experimental/PropertyAnimatorCore/Source/PropertyAnimatorCore/Public/Components/PropertyAnimatorCoreComponent.h b/Engine/Plugins/Experimental/PropertyAnimatorCore/Source/PropertyAnimatorCore/Public/Components/PropertyAnimatorCoreComponent.h
--- a/Engine/Plugins/Experimental/PropertyAnimatorCore/Source/PropertyAnimatorCore/Public/Components/PropertyAnimatorCoreComponent.h
+++ b/Engine/Plugins/Experimental/PropertyAnimatorCore/Source/PropertyAnimatorCore/Public/Components/PropertyAnimatorCoreComponent.h
@@ -109,10 +109,10 @@ protected:
 	void OnAnimatorsSetEnabled(const UWorld* InWorld, bool bInEnabled, bool bInTransact);
 
 	/** Callback when PropertyAnimators changed */
-	void OnAnimatorsChanged();
+	void OnAnimatorsChanged(EPropertyAnimatorCoreUpdateEvent InType);
 
 	/** Callback when global enabled state is changed */
-	void OnAnimatorsEnabledChanged();
+	void OnAnimatorsEnabledChanged(EPropertyAnimatorCoreUpdateEvent InType);
 
 	/** Callback when global time source name is changed */
 	void OnTimeSourceNameChanged();
@@ -120,8 +120,6 @@ protected:
 	/** Evaluate only specified animators */
 	bool EvaluateAnimators();
 
-	void InitAnimatorComponent();
-
 	/** Finds a cached time source with this name or creates a new one */
 	UPropertyAnimatorCoreTimeSourceBase* FindOrAddTimeSource(FName InTimeSourceName);
 
diff --git a/Engine/Plugins/Experimental/PropertyAnimatorCore/Source/PropertyAnimatorCoreEditor/Private/Customizations/PropertyAnimatorCoreEditorStackCustomization.cpp b/Engine/Plugins/Experimental/PropertyAnimatorCore/Source/PropertyAnimatorCoreEditor/Private/Customizations/PropertyAnimatorCoreEditorStackCustomization.cpp
--- a/Engine/Plugins/Experimental/PropertyAnimatorCore/Source/PropertyAnimatorCoreEditor/Private/Customizations/PropertyAnimatorCoreEditorStackCustomization.cpp
+++ b/Engine/Plugins/Experimental/PropertyAnimatorCore/Source/PropertyAnimatorCoreEditor/Private/Customizations/PropertyAnimatorCoreEditorStackCustomization.cpp
@@ -527,15 +527,20 @@ void UPropertyAnimatorCoreEditorStackCustomization::FillAnimatorContextActionMen
 	InToolMenu->AddMenuEntry(RemoveAnimatorMenuEntry.Name, RemoveAnimatorMenuEntry);
 }
 
-void UPropertyAnimatorCoreEditorStackCustomization::OnAnimatorUpdated(UPropertyAnimatorCoreComponent* InComponent, UPropertyAnimatorCoreBase* InAnimator)
+void UPropertyAnimatorCoreEditorStackCustomization::OnAnimatorUpdated(UPropertyAnimatorCoreComponent* InComponent, UPropertyAnimatorCoreBase* InAnimator, EPropertyAnimatorCoreUpdateEvent InType)
 {
 	if (InComponent)
 	{
 		RefreshActiveSelection(InComponent, /** Force */false);
+
+		if (InAnimator && InType == EPropertyAnimatorCoreUpdateEvent::User)
+		{
+			FocusCustomization(InComponent);
+		}
 	}
 }
 
-void UPropertyAnimatorCoreEditorStackCustomization::OnAnimatorRemoved(UPropertyAnimatorCoreComponent* InComponent, UPropertyAnimatorCoreBase* InAnimator)
+void UPropertyAnimatorCoreEditorStackCustomization::OnAnimatorRemoved(UPropertyAnimatorCoreComponent* InComponent, UPropertyAnimatorCoreBase* InAnimator, EPropertyAnimatorCoreUpdateEvent InType)
 {
 	if (InComponent)
 	{
diff --git a/Engine/Plugins/Experimental/PropertyAnimatorCore/Source/PropertyAnimatorCoreEditor/Private/Customizations/PropertyAnimatorCoreEditorStackCustomization.h b/Engine/Plugins/Experimental/PropertyAnimatorCore/Source/PropertyAnimatorCoreEditor/Private/Customizations/PropertyAnimatorCoreEditorStackCustomization.h
--- a/Engine/Plugins/Experimental/PropertyAnimatorCore/Source/PropertyAnimatorCoreEditor/Private/Customizations/PropertyAnimatorCoreEditorStackCustomization.h
+++ b/Engine/Plugins/Experimental/PropertyAnimatorCore/Source/PropertyAnimatorCoreEditor/Private/Customizations/PropertyAnimatorCoreEditorStackCustomization.h
@@ -8,6 +8,7 @@
 class UPropertyAnimatorCoreComponent;
 class UPropertyAnimatorCoreBase;
 class UToolMenu;
+enum class EPropertyAnimatorCoreUpdateEvent : uint8;
 
 /** Property Controller customization for operator stack tab */
 UCLASS()
@@ -42,8 +43,8 @@ protected:
 	void FillAnimatorHeaderActionMenu(UToolMenu* InToolMenu);
 	void FillAnimatorContextActionMenu(UToolMenu* InToolMenu) const;
 
-	void OnAnimatorUpdated(UPropertyAnimatorCoreComponent* InComponent, UPropertyAnimatorCoreBase* InAnimator);
-	void OnAnimatorRemoved(UPropertyAnimatorCoreComponent* InComponent, UPropertyAnimatorCoreBase* InAnimator);
+	void OnAnimatorUpdated(UPropertyAnimatorCoreComponent* InComponent, UPropertyAnimatorCoreBase* InAnimator, EPropertyAnimatorCoreUpdateEvent InType);
+	void OnAnimatorRemoved(UPropertyAnimatorCoreComponent* InComponent, UPropertyAnimatorCoreBase* InAnimator, EPropertyAnimatorCoreUpdateEvent InType);
 
 	/** Create commands for animator */
 	TSharedRef<FUICommandList> CreateAnimatorCommands(FOperatorStackEditorItemPtr InItem);
