MotionDesign : PropertyAnimator
- Reduced seed button text size + set details panel font

diff --git a/Engine/Plugins/Experimental/PropertyAnimatorCore/Source/PropertyAnimatorCoreEditor/Private/Customizations/PropertyAnimatorCoreEditorSeedTypeCustomization.cpp b/Engine/Plugins/Experimental/PropertyAnimatorCore/Source/PropertyAnimatorCoreEditor/Private/Customizations/PropertyAnimatorCoreEditorSeedTypeCustomization.cpp
--- a/Engine/Plugins/Experimental/PropertyAnimatorCore/Source/PropertyAnimatorCoreEditor/Private/Customizations/PropertyAnimatorCoreEditorSeedTypeCustomization.cpp
+++ b/Engine/Plugins/Experimental/PropertyAnimatorCore/Source/PropertyAnimatorCoreEditor/Private/Customizations/PropertyAnimatorCoreEditorSeedTypeCustomization.cpp
@@ -8,6 +8,7 @@
 #include "Templates/SharedPointer.h"
 #include "Widgets/Input/SButton.h"
 #include "Widgets/SBoxPanel.h"
+#include "Widgets/Text/STextBlock.h"
 
 #define LOCTEXT_NAMESPACE "PropertyAnimatorCoreEditorSeedTypeCustomization"
 
@@ -33,9 +34,13 @@ void FPropertyAnimatorCoreEditorSeedTypeCustomization::CustomizeHeader(TSharedRe
 		.AutoWidth()
 		[
 			SNew(SButton)
-			.Text(LOCTEXT("GenerateSeed", "Generate Seed"))
-			.ToolTipText(LOCTEXT("GenerateSeedTooltip", "Generates a new seed between property type min and max value"))
 			.OnClicked_Static(&FPropertyAnimatorCoreEditorSeedTypeCustomization::OnGenerateSeedClicked, InPropertyHandle.ToWeakPtr())
+			[
+				SNew(STextBlock)
+				.Font(InUtils.GetRegularFont())
+				.Text(LOCTEXT("GenerateSeed", "Seed"))
+				.ToolTipText(LOCTEXT("GenerateSeedTooltip", "Generates a new seed between property type min and max value"))
+			]
 		]
 	];
 }
