MotionDesign : Animators
- Add multi selection add animators + (un)link support for properties and presets in operator stack and details view

diff --git a/Engine/Plugins/Experimental/PropertyAnimatorCore/Source/PropertyAnimatorCore/Public/Properties/PropertyAnimatorCoreData.h b/Engine/Plugins/Experimental/PropertyAnimatorCore/Source/PropertyAnimatorCore/Public/Properties/PropertyAnimatorCoreData.h
--- a/Engine/Plugins/Experimental/PropertyAnimatorCore/Source/PropertyAnimatorCore/Public/Properties/PropertyAnimatorCoreData.h
+++ b/Engine/Plugins/Experimental/PropertyAnimatorCore/Source/PropertyAnimatorCore/Public/Properties/PropertyAnimatorCoreData.h
@@ -107,7 +107,7 @@ struct FPropertyAnimatorCoreData
 	}
 
 	/** Creates a single path that contains resolver, owner and properties for (de)serialization */
-	FString GetPropertyLocatorPath() const;
+	PROPERTYANIMATORCORE_API FString GetPropertyLocatorPath() const;
 
 	/** The member property name */
 	PROPERTYANIMATORCORE_API FName GetMemberPropertyName() const;
diff --git a/Engine/Plugins/Experimental/PropertyAnimatorCore/Source/PropertyAnimatorCoreEditor/Private/Customizations/PropertyAnimatorCoreEditorDetailCustomization.cpp b/Engine/Plugins/Experimental/PropertyAnimatorCore/Source/PropertyAnimatorCoreEditor/Private/Customizations/PropertyAnimatorCoreEditorDetailCustomization.cpp
--- a/Engine/Plugins/Experimental/PropertyAnimatorCore/Source/PropertyAnimatorCoreEditor/Private/Customizations/PropertyAnimatorCoreEditorDetailCustomization.cpp
+++ b/Engine/Plugins/Experimental/PropertyAnimatorCore/Source/PropertyAnimatorCoreEditor/Private/Customizations/PropertyAnimatorCoreEditorDetailCustomization.cpp
@@ -4,6 +4,7 @@
 
 #include "DetailLayoutBuilder.h"
 #include "DetailWidgetRow.h"
+#include "Menus/PropertyAnimatorCoreEditorMenuContext.h"
 #include "Presets/PropertyAnimatorCorePropertyPreset.h"
 #include "Styles/PropertyAnimatorCoreEditorStyle.h"
 #include "Subsystems/PropertyAnimatorCoreEditorSubsystem.h"
@@ -30,16 +31,13 @@ void FPropertyAnimatorCoreEditorDetailCustomization::CustomizeDetails(IDetailLay
 		return;
 	}
 
-	TArray<TWeakObjectPtr<UPropertyAnimatorCoreBase>> AnimatorsWeak = InDetailBuilder.GetObjectsOfTypeBeingCustomized<UPropertyAnimatorCoreBase>();
-	UPropertyAnimatorCoreBase* Animator = AnimatorsWeak.IsValidIndex(0) ? AnimatorsWeak[0].Get() : nullptr;
+	AnimatorsWeak = InDetailBuilder.GetObjectsOfTypeBeingCustomized<UPropertyAnimatorCoreBase>();
 
-	if (!Animator)
+	if (AnimatorsWeak.IsEmpty())
 	{
 		return;
 	}
 
-	AnimatorWeak = Animator;
-
 	IDetailPropertyRow* PropertyRow = InDetailBuilder.EditDefaultProperty(LinkedPropertiesHandle);
 
 	if (!PropertyRow)
@@ -91,6 +89,7 @@ void FPropertyAnimatorCoreEditorDetailCustomization::CustomizeDetails(IDetailLay
 			.ContentPadding(2.f)
 			.ButtonStyle(&FCoreStyle::Get().GetWidgetStyle<FButtonStyle>("Button"))
 			.ToolTipText(LOCTEXT("CreatePropertyPreset", "Create a preset from these properties"))
+			.Visibility(AnimatorsWeak.Num() == 1 ? EVisibility::Visible : EVisibility::Collapsed)
 			.OnClicked(this, &FPropertyAnimatorCoreEditorDetailCustomization::OnCreatePropertyPresetClicked)
 			[
 				SNew(SImage)
@@ -133,9 +132,8 @@ void FPropertyAnimatorCoreEditorDetailCustomization::CustomizeDetails(IDetailLay
 TSharedRef<SWidget> FPropertyAnimatorCoreEditorDetailCustomization::GenerateLinkMenu()
 {
 	UToolMenus* ToolMenus = UToolMenus::Get();
-	UPropertyAnimatorCoreBase* Animator = AnimatorWeak.Get();
 
-	if (!Animator || !ToolMenus)
+	if (AnimatorsWeak.IsEmpty() || !ToolMenus)
 	{
 		return SNullWidget::NullWidget;
 	}
@@ -148,20 +146,47 @@ TSharedRef<SWidget> FPropertyAnimatorCoreEditorDetailCustomization::GenerateLink
 		LinkPropertiesMenu->AddDynamicSection(TEXT("FillLinkPropertiesCustomizationMenu"), FNewToolMenuDelegate::CreateStatic(&FPropertyAnimatorCoreEditorDetailCustomization::FillLinkMenu));
 	}
 
-	return ToolMenus->GenerateWidget(MenuName, FToolMenuContext(Animator));
+	UPropertyAnimatorCoreEditorMenuContext* MenuContext = NewObject<UPropertyAnimatorCoreEditorMenuContext>();
+
+	TSet<UPropertyAnimatorCoreBase*> Animators;
+	Algo::Transform(AnimatorsWeak, Animators, [](const TWeakObjectPtr<UPropertyAnimatorCoreBase>& InAnimatorWeak)
+	{
+		return InAnimatorWeak.Get();
+	});
+	MenuContext->SetAnimators(Animators);
+	return ToolMenus->GenerateWidget(MenuName, FToolMenuContext(MenuContext));
 }
 
 void FPropertyAnimatorCoreEditorDetailCustomization::FillLinkMenu(UToolMenu* InToolMenu)
 {
 	UPropertyAnimatorCoreEditorSubsystem* AnimatorEditorSubsystem = UPropertyAnimatorCoreEditorSubsystem::Get();
-	UPropertyAnimatorCoreBase* Animator = InToolMenu->FindContext<UPropertyAnimatorCoreBase>();
+	const UPropertyAnimatorCoreEditorMenuContext* MenuContextObject = InToolMenu->FindContext<UPropertyAnimatorCoreEditorMenuContext>();
+
+	if (!AnimatorEditorSubsystem || !MenuContextObject)
+	{
+		return;
+	}
 
-	if (!Animator || !AnimatorEditorSubsystem)
+	TSet<UObject*> ContextObjects;
+	Algo::TransformIf(
+		MenuContextObject->GetAnimators(),
+		ContextObjects,
+		[](const UPropertyAnimatorCoreBase* InAnimator)
+		{
+			return IsValid(InAnimator);
+		},
+		[](UPropertyAnimatorCoreBase* InAnimator)
+		{
+			return InAnimator;
+		}
+	);
+
+	if (ContextObjects.IsEmpty())
 	{
 		return;
 	}
 
-	const FPropertyAnimatorCoreEditorMenuContext MenuContext({Animator}, {});
+	const FPropertyAnimatorCoreEditorMenuContext MenuContext(ContextObjects, {});
 	FPropertyAnimatorCoreEditorMenuOptions MenuOptions(
 		{
 			EPropertyAnimatorCoreEditorMenuType::Link
@@ -174,87 +199,114 @@ void FPropertyAnimatorCoreEditorDetailCustomization::FillLinkMenu(UToolMenu* InT
 
 bool FPropertyAnimatorCoreEditorDetailCustomization::IsAnyPropertyLinked() const
 {
-	const UPropertyAnimatorCoreBase* Animator = AnimatorWeak.Get();
-	if (!Animator)
+	for (const TWeakObjectPtr<UPropertyAnimatorCoreBase>& AnimatorWeak : AnimatorsWeak)
 	{
-		return false;
+		const UPropertyAnimatorCoreBase* Animator = AnimatorWeak.Get();
+		if (!Animator)
+		{
+			continue;
+		}
+
+		if (Animator->GetLinkedPropertiesCount() > 0)
+		{
+			return true;
+		}
 	}
 
-	return Animator->GetLinkedPropertiesCount() > 0;
+	return false;
 }
 
 ECheckBoxState FPropertyAnimatorCoreEditorDetailCustomization::IsPropertiesEnabled() const
 {
-	ECheckBoxState State = ECheckBoxState::Undetermined;
-
-	const UPropertyAnimatorCoreBase* Animator = AnimatorWeak.Get();
-	if (!Animator)
-	{
-		return State;
-	}
-
-	int32 EnabledProperties = 0;
-	constexpr bool bResolve = false;
-	Animator->ForEachLinkedProperty<UPropertyAnimatorCoreContext>([&EnabledProperties](UPropertyAnimatorCoreContext* InOptions, const FPropertyAnimatorCoreData& InProperty)->bool
+	TOptional<ECheckBoxState> State;
+	for (const TWeakObjectPtr<UPropertyAnimatorCoreBase>& AnimatorWeak : AnimatorsWeak)
 	{
-		if (InOptions->IsAnimated())
+		UPropertyAnimatorCoreBase* Animator = AnimatorWeak.Get();
+		if (!Animator)
 		{
-			EnabledProperties++;
+			continue;
 		}
-		return true;
-	}, bResolve);
 
-	const int32 LinkedProperties = Animator->GetLinkedPropertiesCount();
+		int32 EnabledProperties = 0;
+		constexpr bool bResolve = false;
+		Animator->ForEachLinkedProperty<UPropertyAnimatorCoreContext>([&EnabledProperties](UPropertyAnimatorCoreContext* InOptions, const FPropertyAnimatorCoreData& InProperty)->bool
+		{
+			if (InOptions->IsAnimated())
+			{
+				EnabledProperties++;
+			}
+			return true;
+		}, bResolve);
+		
+		const int32 LinkedProperties = Animator->GetLinkedPropertiesCount();
+
+		ECheckBoxState AnimatorState = ECheckBoxState::Undetermined;
+		if (LinkedProperties > 0 && EnabledProperties == 0)
+		{
+			AnimatorState = ECheckBoxState::Unchecked;
+		}
+		else if (LinkedProperties > 0 && EnabledProperties < LinkedProperties)
+		{
+			AnimatorState = ECheckBoxState::Undetermined;
+		}
+		else
+		{
+			AnimatorState = ECheckBoxState::Checked;
+		}
 
-	if (LinkedProperties > 0 && EnabledProperties == 0)
-	{
-		State = ECheckBoxState::Unchecked;
-	}
-	else if (LinkedProperties > 0 && EnabledProperties < LinkedProperties)
-	{
-		State = ECheckBoxState::Undetermined;
-	}
-	else
-	{
-		State = ECheckBoxState::Checked;
+		if (!State.IsSet())
+		{
+			State = AnimatorState;
+		}
+		else if (State.GetValue() != AnimatorState)
+		{
+			State = ECheckBoxState::Undetermined;
+			break;
+		}
 	}
 
-	return State;
+	return State.Get(ECheckBoxState::Undetermined);
 }
 
 void FPropertyAnimatorCoreEditorDetailCustomization::OnPropertiesEnabled(ECheckBoxState InNewState) const
 {
-	const UPropertyAnimatorCoreBase* Animator = AnimatorWeak.Get();
-	if (!Animator)
-	{
-		return;
-	}
-
 	if (InNewState == ECheckBoxState::Undetermined)
 	{
 		return;
 	}
 
-	constexpr bool bResolve = false;
-	Animator->ForEachLinkedProperty<UPropertyAnimatorCoreContext>([InNewState](UPropertyAnimatorCoreContext* InOptions, const FPropertyAnimatorCoreData& InProperty)->bool
+	for (const TWeakObjectPtr<UPropertyAnimatorCoreBase>& AnimatorWeak : AnimatorsWeak)
 	{
-		InOptions->SetAnimated(InNewState == ECheckBoxState::Checked);
-		return true;
-	}, bResolve);
+		const UPropertyAnimatorCoreBase* Animator = AnimatorWeak.Get();
+		if (!Animator)
+		{
+			continue;
+		}
+
+		constexpr bool bResolve = false;
+		Animator->ForEachLinkedProperty<UPropertyAnimatorCoreContext>([InNewState](UPropertyAnimatorCoreContext* InOptions, const FPropertyAnimatorCoreData& InProperty)->bool
+		{
+			InOptions->SetAnimated(InNewState == ECheckBoxState::Checked);
+			return true;
+		}, bResolve);
+	}
 }
 
 FReply FPropertyAnimatorCoreEditorDetailCustomization::UnlinkProperties() const
 {
-	UPropertyAnimatorCoreBase* Animator = AnimatorWeak.Get();
-	if (!Animator)
-	{
-		return FReply::Handled();
-	}
-
 	if (UPropertyAnimatorCoreSubsystem* AnimatorSubsystem = UPropertyAnimatorCoreSubsystem::Get())
 	{
-		const TSet<FPropertyAnimatorCoreData> LinkedProperties = Animator->GetLinkedProperties();
-		AnimatorSubsystem->UnlinkAnimatorProperties(Animator, LinkedProperties, true);
+		for (const TWeakObjectPtr<UPropertyAnimatorCoreBase>& AnimatorWeak : AnimatorsWeak)
+		{
+			UPropertyAnimatorCoreBase* Animator = AnimatorWeak.Get();	
+			if (!Animator)
+			{
+				continue;
+			}
+		
+			const TSet<FPropertyAnimatorCoreData> LinkedProperties = Animator->GetLinkedProperties();
+			AnimatorSubsystem->UnlinkAnimatorProperties(Animator, LinkedProperties, true);
+		}
 	}
 
 	return FReply::Handled();
@@ -262,7 +314,12 @@ FReply FPropertyAnimatorCoreEditorDetailCustomization::UnlinkProperties() const
 
 FReply FPropertyAnimatorCoreEditorDetailCustomization::OnCreatePropertyPresetClicked() const
 {
-	UPropertyAnimatorCoreBase* Animator = AnimatorWeak.Get();
+	if (AnimatorsWeak.IsEmpty())
+	{
+		return FReply::Handled();
+	}
+
+	const UPropertyAnimatorCoreBase* Animator = AnimatorsWeak[0].Get();
 	if (!Animator)
 	{
 		return FReply::Handled();
diff --git a/Engine/Plugins/Experimental/PropertyAnimatorCore/Source/PropertyAnimatorCoreEditor/Private/Customizations/PropertyAnimatorCoreEditorDetailCustomization.h b/Engine/Plugins/Experimental/PropertyAnimatorCore/Source/PropertyAnimatorCoreEditor/Private/Customizations/PropertyAnimatorCoreEditorDetailCustomization.h
--- a/Engine/Plugins/Experimental/PropertyAnimatorCore/Source/PropertyAnimatorCoreEditor/Private/Customizations/PropertyAnimatorCoreEditorDetailCustomization.h
+++ b/Engine/Plugins/Experimental/PropertyAnimatorCore/Source/PropertyAnimatorCoreEditor/Private/Customizations/PropertyAnimatorCoreEditorDetailCustomization.h
@@ -38,5 +38,5 @@ protected:
 
 	FReply OnCreatePropertyPresetClicked() const;
 
-	TWeakObjectPtr<UPropertyAnimatorCoreBase> AnimatorWeak;
+	TArray<TWeakObjectPtr<UPropertyAnimatorCoreBase>> AnimatorsWeak;
 };
diff --git a/Engine/Plugins/Experimental/PropertyAnimatorCore/Source/PropertyAnimatorCoreEditor/Private/Menus/PropertyAnimatorCoreEditorMenu.cpp b/Engine/Plugins/Experimental/PropertyAnimatorCore/Source/PropertyAnimatorCoreEditor/Private/Menus/PropertyAnimatorCoreEditorMenu.cpp
--- a/Engine/Plugins/Experimental/PropertyAnimatorCore/Source/PropertyAnimatorCoreEditor/Private/Menus/PropertyAnimatorCoreEditorMenu.cpp
+++ b/Engine/Plugins/Experimental/PropertyAnimatorCore/Source/PropertyAnimatorCoreEditor/Private/Menus/PropertyAnimatorCoreEditorMenu.cpp
@@ -125,25 +125,37 @@ void UE::PropertyAnimatorCoreEditor::Menu::FillExistingAnimatorSection(UToolMenu
 	constexpr bool bCloseMenuAfterSelection = true;
 	constexpr bool bOpenOnClick = false;
 
+	TMap<FName, TSet<UPropertyAnimatorCoreBase*>> ExistingAnimatorsMap;
 	for (const FPropertyAnimatorCoreData& Property : InMenuData->GetContext().GetProperties())
 	{
-		for (UPropertyAnimatorCoreBase* Animator : Subsystem->GetExistingAnimators(Property))
+		for (AActor* Actor : InMenuData->GetContext().GetActors())
 		{
-			const FName MenuName = Animator->GetAnimatorDisplayName();
-			const FText MenuLabel = FText::FromName(MenuName);
-			const FSlateIcon MenuIcon = FSlateIconFinder::FindIconForClass(Animator->GetClass());
-
-			ExistingAnimatorsSection.AddSubMenu(
-				MenuName
-				, MenuLabel
-				, LOCTEXT("ExistingAnimatorSection.Tooltip", "Link or unlink properties for this animator")
-				, FNewToolMenuDelegate::CreateLambda(&FillLinkAnimatorSubmenu, Animator, InMenuData)
-				, bOpenOnClick
-				, MenuIcon
-				, bCloseMenuAfterSelection
-			);
+			const FPropertyAnimatorCoreData ActorProperty(Actor, Property.GetPropertyLocatorPath());
+			for (UPropertyAnimatorCoreBase* Animator : Subsystem->GetExistingAnimators(ActorProperty))
+			{
+				const FName AnimatorKey = FName(Animator->GetAnimatorDisplayName().ToString() + TEXT(" (") + Animator->GetClass()->GetName() + TEXT(")"));
+				ExistingAnimatorsMap.FindOrAdd(AnimatorKey).Add(Animator);
+			}
 		}
 	}
+
+	for (const TPair<FName, TSet<UPropertyAnimatorCoreBase*>>& ExistingAnimators : ExistingAnimatorsMap)
+	{
+		UPropertyAnimatorCoreBase* Animator = ExistingAnimators.Value.Array()[0];
+		const FName MenuName = ExistingAnimators.Key;
+		const FText MenuLabel = FText::FromName(ExistingAnimators.Key);
+		const FSlateIcon MenuIcon = FSlateIconFinder::FindIconForClass(Animator->GetClass());
+
+		ExistingAnimatorsSection.AddSubMenu(
+			MenuName
+			, MenuLabel
+			, LOCTEXT("ExistingAnimatorSection.Tooltip", "Link or unlink properties for this animator")
+			, FNewToolMenuDelegate::CreateLambda(&FillLinkAnimatorSubmenu, ExistingAnimators.Value, InMenuData)
+			, bOpenOnClick
+			, MenuIcon
+			, bCloseMenuAfterSelection
+		);
+	}
 }
 
 void UE::PropertyAnimatorCoreEditor::Menu::FillLinkAnimatorSection(UToolMenu* InMenu, TSharedRef<FPropertyAnimatorCoreEditorMenuData> InMenuData)
@@ -161,19 +173,12 @@ void UE::PropertyAnimatorCoreEditor::Menu::FillLinkAnimatorSection(UToolMenu* In
 
 	const TSet<UPropertyAnimatorCoreBase*> Animators = InMenuData->GetContext().GetAnimators();
 
-	if (Animators.Num() != 1)
-	{
-		return;
-	}
-
-	UPropertyAnimatorCoreBase* Animator = Animators.Array()[0];
-
-	if (!Animator)
+	if (Animators.IsEmpty())
 	{
 		return;
 	}
 
-	FillLinkAnimatorSubmenu(InMenu, Animator, InMenuData);
+	FillLinkAnimatorSubmenu(InMenu, Animators, InMenuData);
 }
 
 void UE::PropertyAnimatorCoreEditor::Menu::FillDeleteAnimatorSection(UToolMenu* InMenu, TSharedRef<FPropertyAnimatorCoreEditorMenuData> InMenuData)
@@ -378,7 +383,7 @@ void UE::PropertyAnimatorCoreEditor::Menu::ExecuteNewAnimatorPresetAction(const
 	InMenuData->SetLastCreatedAnimators(Subsystem->CreateAnimators(InActors, InAnimator->GetClass(), InPreset, InMenuData->GetOptions().ShouldTransact()));
 }
 
-void UE::PropertyAnimatorCoreEditor::Menu::ExecuteNewAnimatorPropertyAction(const UPropertyAnimatorCoreBase* InAnimator, FPropertyAnimatorCoreData InProperty, TSharedRef<FPropertyAnimatorCoreEditorMenuData> InMenuData)
+void UE::PropertyAnimatorCoreEditor::Menu::ExecuteNewAnimatorPropertyAction(const UPropertyAnimatorCoreBase* InAnimator, const FString& InPropertyLocatorPath, TSharedRef<FPropertyAnimatorCoreEditorMenuData> InMenuData)
 {
 	UPropertyAnimatorCoreSubsystem* Subsystem = UPropertyAnimatorCoreSubsystem::Get();
 
@@ -390,14 +395,19 @@ void UE::PropertyAnimatorCoreEditor::Menu::ExecuteNewAnimatorPropertyAction(cons
 		return;
 	}
 
-	AActor* OwningActor = InProperty.GetOwningActor();
-
-	UPropertyAnimatorCoreBase* NewAnimator = Subsystem->CreateAnimator(OwningActor, InAnimator->GetClass(), nullptr, InMenuData->GetOptions().ShouldTransact());
-	Subsystem->LinkAnimatorProperty(NewAnimator, InProperty, true);
-	InMenuData->SetLastCreatedAnimator(NewAnimator);
+	TSet<UPropertyAnimatorCoreBase*> CreatedAnimators;
+	for (AActor* Actor : InMenuData->GetContext().GetActors())
+	{
+		UPropertyAnimatorCoreBase* NewAnimator = Subsystem->CreateAnimator(Actor, InAnimator->GetClass(), nullptr, InMenuData->GetOptions().ShouldTransact());
+		FPropertyAnimatorCoreData Property(Actor, InPropertyLocatorPath);
+		Subsystem->LinkAnimatorProperty(NewAnimator, Property, true);
+		CreatedAnimators.Add(NewAnimator);
+	}
+	
+	InMenuData->SetLastCreatedAnimators(CreatedAnimators);
 }
 
-void UE::PropertyAnimatorCoreEditor::Menu::ExecuteLinkLastCreatedAnimatorPropertyAction(const UPropertyAnimatorCoreBase* InAnimator, FPropertyAnimatorCoreData InProperty, TSharedRef<FPropertyAnimatorCoreEditorMenuData> InMenuData)
+void UE::PropertyAnimatorCoreEditor::Menu::ExecuteLinkLastCreatedAnimatorPropertyAction(const UPropertyAnimatorCoreBase* InAnimator, const FString& InPropertyLocatorPath, TSharedRef<FPropertyAnimatorCoreEditorMenuData> InMenuData)
 {
 	UPropertyAnimatorCoreSubsystem* Subsystem = UPropertyAnimatorCoreSubsystem::Get();
 
@@ -408,6 +418,7 @@ void UE::PropertyAnimatorCoreEditor::Menu::ExecuteLinkLastCreatedAnimatorPropert
 		return;
 	}
 
+	constexpr bool bShouldTransact = true;
 	for (UPropertyAnimatorCoreBase* LastCreatedAnimator : InMenuData->GetLastCreatedAnimators())
 	{
 		if (LastCreatedAnimator->GetClass() != InAnimator->GetClass())
@@ -415,13 +426,14 @@ void UE::PropertyAnimatorCoreEditor::Menu::ExecuteLinkLastCreatedAnimatorPropert
 			continue;
 		}
 
-		if (LastCreatedAnimator->IsPropertyLinked(InProperty))
+		FPropertyAnimatorCoreData Property(LastCreatedAnimator->GetAnimatorActor(), InPropertyLocatorPath);
+		if (LastCreatedAnimator->IsPropertyLinked(Property))
 		{
-			Subsystem->UnlinkAnimatorProperty(LastCreatedAnimator, InProperty, true);
+			Subsystem->UnlinkAnimatorProperty(LastCreatedAnimator, Property, bShouldTransact);
 		}
 		else
 		{
-			Subsystem->LinkAnimatorProperty(LastCreatedAnimator, InProperty, true);
+			Subsystem->LinkAnimatorProperty(LastCreatedAnimator, Property, bShouldTransact);
 		}
 	}
 }
@@ -455,59 +467,82 @@ void UE::PropertyAnimatorCoreEditor::Menu::ExecuteApplyLastCreatedAnimatorPreset
 	}
 }
 
-void UE::PropertyAnimatorCoreEditor::Menu::ExecuteLinkAnimatorPresetAction(UPropertyAnimatorCoreBase* InAnimator, UPropertyAnimatorCorePropertyPreset* InPreset, TSharedRef<FPropertyAnimatorCoreEditorMenuData> InMenuData)
+void UE::PropertyAnimatorCoreEditor::Menu::ExecuteLinkAnimatorPresetAction(const TSet<UPropertyAnimatorCoreBase*>& InAnimators, UPropertyAnimatorCorePropertyPreset* InPreset, TSharedRef<FPropertyAnimatorCoreEditorMenuData> InMenuData)
 {
 	UPropertyAnimatorCoreSubsystem* Subsystem = UPropertyAnimatorCoreSubsystem::Get();
 
 	if (!Subsystem
 		|| !InPreset
-		|| !IsValid(InAnimator)
-		|| InAnimator->IsTemplate()
+		|| InAnimators.IsEmpty()
 		|| InMenuData->GetContext().IsEmpty())
 	{
 		return;
 	}
 
-	if (InPreset->IsPresetApplied(InAnimator))
-	{
-		Subsystem->UnapplyAnimatorPreset(InAnimator, InPreset, InMenuData->GetOptions().ShouldTransact());
-	}
-	else
+	const ECheckBoxState State = GetAnimatorPresetState(InAnimators, InPreset);
+
+	for (UPropertyAnimatorCoreBase* Animator : InAnimators)
 	{
-		Subsystem->ApplyAnimatorPreset(InAnimator, InPreset, InMenuData->GetOptions().ShouldTransact());
+		if (!IsValid(Animator) || Animator->IsTemplate())
+		{
+			continue;
+		}
+
+		if (State == ECheckBoxState::Undetermined || InPreset->IsPresetApplied(Animator))
+		{
+			Subsystem->UnapplyAnimatorPreset(Animator, InPreset, InMenuData->GetOptions().ShouldTransact());
+		}
+		else
+		{
+			Subsystem->ApplyAnimatorPreset(Animator, InPreset, InMenuData->GetOptions().ShouldTransact());
+		}
 	}
 }
 
-void UE::PropertyAnimatorCoreEditor::Menu::ExecuteLinkAnimatorPropertyAction(UPropertyAnimatorCoreBase* InAnimator, FPropertyAnimatorCoreData InProperty, UPropertyAnimatorCorePropertyPreset* InPreset, TSharedRef<FPropertyAnimatorCoreEditorMenuData> InMenuData)
+void UE::PropertyAnimatorCoreEditor::Menu::ExecuteLinkAnimatorPropertyAction(const TSet<UPropertyAnimatorCoreBase*>& InAnimators, const FString& InPropertyLocatorPath, UPropertyAnimatorCorePropertyPreset* InPreset, TSharedRef<FPropertyAnimatorCoreEditorMenuData> InMenuData)
 {
 	UPropertyAnimatorCoreSubsystem* Subsystem = UPropertyAnimatorCoreSubsystem::Get();
 
 	if (!Subsystem
-		|| !IsValid(InAnimator)
-		|| InAnimator->IsTemplate()
-		|| !InProperty.IsResolved()
-		|| !InAnimator->HasPropertySupport(InProperty)
+		|| InAnimators.IsEmpty()
 		|| InMenuData->GetContext().IsEmpty())
 	{
 		return;
 	}
 
-	if (InAnimator->IsPropertyLinked(InProperty))
+	const ECheckBoxState State = GetAnimatorPropertyLinkState(InAnimators, InPropertyLocatorPath);
+
+	for (UPropertyAnimatorCoreBase* Animator : InAnimators)
 	{
-		Subsystem->UnlinkAnimatorProperty(InAnimator, InProperty, InMenuData->GetOptions().ShouldTransact());
+		if (!IsValid(Animator) || Animator->IsTemplate())
+		{
+			continue;
+		}
+
+		FPropertyAnimatorCoreData AnimatorProperty(Animator->GetAnimatorActor(), InPropertyLocatorPath);
 
-		if (InPreset)
+		if (!AnimatorProperty.IsResolved() || !Animator->HasPropertySupport(AnimatorProperty))
 		{
-			InPreset->OnPresetUnapplied(InAnimator, {InProperty});
+			continue;
 		}
-	}
-	else
-	{
-		Subsystem->LinkAnimatorProperty(InAnimator, InProperty, InMenuData->GetOptions().ShouldTransact());
 
-		if (InPreset)
+		if (State == ECheckBoxState::Undetermined || Animator->IsPropertyLinked(AnimatorProperty))
+		{
+			Subsystem->UnlinkAnimatorProperty(Animator, AnimatorProperty, InMenuData->GetOptions().ShouldTransact());
+
+			if (InPreset)
+			{
+				InPreset->OnPresetUnapplied(Animator, {AnimatorProperty});
+			}
+		}
+		else
 		{
-			InPreset->OnPresetApplied(InAnimator, {InProperty});
+			Subsystem->LinkAnimatorProperty(Animator, AnimatorProperty, InMenuData->GetOptions().ShouldTransact());
+
+			if (InPreset)
+			{
+				InPreset->OnPresetApplied(Animator, {AnimatorProperty});
+			}
 		}
 	}
 }
@@ -591,16 +626,30 @@ void UE::PropertyAnimatorCoreEditor::Menu::ExecuteDeleteAnimatorAction(UProperty
 	Subsystem->RemoveAnimator(InAnimator, InMenuData->GetOptions().ShouldTransact());
 }
 
-ECheckBoxState UE::PropertyAnimatorCoreEditor::Menu::GetAnimatorPresetState(const UPropertyAnimatorCoreBase* InAnimator, UPropertyAnimatorCorePropertyPreset* InPreset)
+ECheckBoxState UE::PropertyAnimatorCoreEditor::Menu::GetAnimatorPresetState(const TSet<UPropertyAnimatorCoreBase*>& InAnimators, UPropertyAnimatorCorePropertyPreset* InPreset)
 {
-	if (!IsValid(InAnimator) || InAnimator->IsTemplate() || !InPreset)
+	if (InAnimators.IsEmpty() || !InPreset)
 	{
 		return ECheckBoxState::Unchecked;
 	}
 
 	TSet<FPropertyAnimatorCoreData> SupportedProperties;
 	TSet<FPropertyAnimatorCoreData> AppliedProperties;
-	InPreset->GetAppliedPresetProperties(InAnimator, SupportedProperties, AppliedProperties);
+
+	for (UPropertyAnimatorCoreBase* Animator : InAnimators)
+	{
+		if (!IsValid(Animator) || Animator->IsTemplate())
+		{
+			continue;
+		}
+
+		TSet<FPropertyAnimatorCoreData> AnimatorSupportedProperties;
+		TSet<FPropertyAnimatorCoreData> AnimatorAppliedProperties;
+		InPreset->GetAppliedPresetProperties(Animator, AnimatorSupportedProperties, AnimatorAppliedProperties);
+
+		SupportedProperties.Append(AnimatorSupportedProperties);
+		AppliedProperties.Append(AnimatorAppliedProperties);
+	}
 
 	if (!SupportedProperties.IsEmpty() && SupportedProperties.Num() == AppliedProperties.Num())
 	{
@@ -622,15 +671,17 @@ ECheckBoxState UE::PropertyAnimatorCoreEditor::Menu::GetLastAnimatorCreatedPrese
 
 	TOptional<ECheckBoxState> AnimatorsState;
 
-	for (const UPropertyAnimatorCoreBase* LastCreatedAnimator : InMenuData->GetLastCreatedAnimators())
+	for (UPropertyAnimatorCoreBase* LastCreatedAnimator : InMenuData->GetLastCreatedAnimators())
 	{
 		if (LastCreatedAnimator->GetClass() == InAnimator->GetClass())
 		{
+			TSet<UPropertyAnimatorCoreBase*> Animators {LastCreatedAnimator};
+
 			if (!AnimatorsState.IsSet())
 			{
-				AnimatorsState = GetAnimatorPresetState(LastCreatedAnimator, InPreset);
+				AnimatorsState = GetAnimatorPresetState(Animators, InPreset);
 			}
-			else if (GetAnimatorPresetState(LastCreatedAnimator, InPreset) != AnimatorsState.GetValue())
+			else if (GetAnimatorPresetState(Animators, InPreset) != AnimatorsState.GetValue())
 			{
 				AnimatorsState = ECheckBoxState::Undetermined;
 				break;
@@ -641,6 +692,49 @@ ECheckBoxState UE::PropertyAnimatorCoreEditor::Menu::GetLastAnimatorCreatedPrese
 	return AnimatorsState.Get(ECheckBoxState::Unchecked);
 }
 
+ECheckBoxState UE::PropertyAnimatorCoreEditor::Menu::GetAnimatorPropertyLinkState(const TSet<UPropertyAnimatorCoreBase*>& InAnimators, const FString& InPropertyLocatorPath)
+{
+	if (InAnimators.IsEmpty())
+	{
+		return ECheckBoxState::Undetermined;
+	}
+
+	ECheckBoxState State = ECheckBoxState::Undetermined;
+	for (const UPropertyAnimatorCoreBase* Animator : InAnimators)
+	{
+		if (!IsValid(Animator) || Animator->IsTemplate())
+		{
+			continue;
+		}
+
+		ECheckBoxState AnimatorState = ECheckBoxState::Undetermined;
+		FPropertyAnimatorCoreData AnimatorProperty(Animator->GetAnimatorActor(), InPropertyLocatorPath);
+
+		if (IsAnimatorPropertyLinked(Animator, AnimatorProperty))
+		{
+			AnimatorState = ECheckBoxState::Checked;
+		}
+		else
+		{
+			AnimatorState = ECheckBoxState::Unchecked;
+		}
+
+		if (State == ECheckBoxState::Undetermined)
+		{
+			State = AnimatorState;
+			continue;
+		}
+
+		if (AnimatorState != State)
+		{
+			State = ECheckBoxState::Undetermined;
+			break;
+		}
+	}
+
+	return State;
+}
+
 bool UE::PropertyAnimatorCoreEditor::Menu::IsAnimatorPresetLinked(UPropertyAnimatorCoreBase* InAnimator, UPropertyAnimatorCorePropertyPreset* InPreset)
 {
 	if (!IsValid(InAnimator) || InAnimator->IsTemplate() || !InPreset)
@@ -730,6 +824,7 @@ void UE::PropertyAnimatorCoreEditor::Menu::FillNewAnimatorSubmenu(UToolMenu* InM
 		{
 			const FString MenuName = SupportedProperty.GetPropertyDisplayName();
 			const FText MenuLabel = FText::FromString(MenuName + TEXT(" (") + SupportedProperty.GetLeafPropertyTypeName().ToString() + TEXT(")"));
+			const FString PropertyLocatorPath = SupportedProperty.GetPropertyLocatorPath();
 
 			// Create action (creates an animator and links the property)
 			PropertySection.AddMenuEntry(
@@ -738,7 +833,7 @@ void UE::PropertyAnimatorCoreEditor::Menu::FillNewAnimatorSubmenu(UToolMenu* InM
 				, LOCTEXT("NewAnimatorPropertySection.Tooltip", "Create this animator using this property")
 				, FSlateIcon()
 				, FUIAction(
-					FExecuteAction::CreateLambda(&ExecuteNewAnimatorPropertyAction, InAnimator, SupportedProperty, InMenuData)
+					FExecuteAction::CreateLambda(&ExecuteNewAnimatorPropertyAction, InAnimator, PropertyLocatorPath, InMenuData)
 					, FCanExecuteAction()
 					, FIsActionChecked()
 					, FIsActionButtonVisible::CreateLambda(&IsLastAnimatorCreatedActionHidden, InAnimator, InMenuData)
@@ -752,9 +847,9 @@ void UE::PropertyAnimatorCoreEditor::Menu::FillNewAnimatorSubmenu(UToolMenu* InM
 				, LOCTEXT("LinkLastCreatedAnimatorPropertySection.Tooltip", "Link this property to the last created animator")
 				, FSlateIcon()
 				, FUIAction(
-					FExecuteAction::CreateLambda(&ExecuteLinkLastCreatedAnimatorPropertyAction, InAnimator, SupportedProperty, InMenuData)
+					FExecuteAction::CreateLambda(&ExecuteLinkLastCreatedAnimatorPropertyAction, InAnimator, PropertyLocatorPath, InMenuData)
 					, FCanExecuteAction()
-					, FIsActionChecked::CreateLambda(&IsLastAnimatorCreatedPropertyLinked, InAnimator, SupportedProperty, InMenuData)
+					, FIsActionChecked::CreateLambda(&IsLastAnimatorCreatedPropertyLinked, InAnimator, PropertyLocatorPath, InMenuData)
 					, FIsActionButtonVisible::CreateLambda(&IsLastAnimatorCreatedActionVisible, InAnimator, InMenuData)
 				)
 				, EUserInterfaceActionType::ToggleButton
@@ -763,11 +858,10 @@ void UE::PropertyAnimatorCoreEditor::Menu::FillNewAnimatorSubmenu(UToolMenu* InM
 	}
 }
 
-void UE::PropertyAnimatorCoreEditor::Menu::FillLinkAnimatorSubmenu(UToolMenu* InMenu, UPropertyAnimatorCoreBase* InAnimator, TSharedRef<FPropertyAnimatorCoreEditorMenuData> InMenuData)
+void UE::PropertyAnimatorCoreEditor::Menu::FillLinkAnimatorSubmenu(UToolMenu* InMenu, const TSet<UPropertyAnimatorCoreBase*>& InAnimators, TSharedRef<FPropertyAnimatorCoreEditorMenuData> InMenuData)
 {
 	if (!InMenu
-		|| !IsValid(InAnimator)
-		|| InAnimator->IsTemplate()
+		|| InAnimators.IsEmpty()
 		|| InMenuData->GetContext().IsEmpty())
 	{
 		return;
@@ -783,9 +877,14 @@ void UE::PropertyAnimatorCoreEditor::Menu::FillLinkAnimatorSubmenu(UToolMenu* In
 
 	TSet<UPropertyAnimatorCorePresetBase*> SupportedPresets = Subsystem->GetAvailablePresets(UPropertyAnimatorCorePropertyPreset::StaticClass());
 
-	for (const AActor* Actor : InMenuData->GetContext().GetActors())
+	for (const UPropertyAnimatorCoreBase* Animator : InAnimators)
 	{
-		SupportedPresets = SupportedPresets.Intersect(Subsystem->GetSupportedPresets(Actor, InAnimator, UPropertyAnimatorCorePropertyPreset::StaticClass()));
+		if (!IsValid(Animator) || Animator->IsTemplate())
+		{
+			continue;
+		}
+
+		SupportedPresets = SupportedPresets.Intersect(Subsystem->GetSupportedPresets(Animator->GetAnimatorActor(), Animator, UPropertyAnimatorCorePropertyPreset::StaticClass()));
 	}
 
 	constexpr bool bCloseMenuAfterSelection = false;
@@ -806,7 +905,7 @@ void UE::PropertyAnimatorCoreEditor::Menu::FillLinkAnimatorSubmenu(UToolMenu* In
 			MenuName
 			, MenuLabel
 			, MenuTooltip
-			, FNewToolMenuDelegate::CreateLambda(&FillPresetAnimatorSubmenu, InAnimator, Cast<UPropertyAnimatorCorePropertyPreset>(SupportedPreset), InMenuData)
+			, FNewToolMenuDelegate::CreateLambda(&FillPresetAnimatorSubmenu, InAnimators, Cast<UPropertyAnimatorCorePropertyPreset>(SupportedPreset), InMenuData)
 			, bOpenOnClick
 			, FSlateIcon()
 			, bCloseMenuAfterSelection
@@ -815,10 +914,25 @@ void UE::PropertyAnimatorCoreEditor::Menu::FillLinkAnimatorSubmenu(UToolMenu* In
 
 	FToolMenuSection& PropertySection = InMenu->FindOrAddSection(TEXT("Properties"), LOCTEXT("LinkAnimatorPropertiesSection.Label", "Properties"));
 
+	bool bFirstAnimator = true;
 	TSet<FPropertyAnimatorCoreData> SupportedProperties;
-	for (const FPropertyAnimatorCoreData& Property : InMenuData->GetContext().GetProperties())
+	for (const UPropertyAnimatorCoreBase* Animator : InAnimators)
 	{
-		InAnimator->GetPropertiesSupported(Property, SupportedProperties, /** SearchDepth */3);
+		for (const FPropertyAnimatorCoreData& Property : InMenuData->GetContext().GetProperties())
+		{
+			TSet<FPropertyAnimatorCoreData> AnimatorSupportedProperties;
+			Animator->GetPropertiesSupported(Property, AnimatorSupportedProperties, /** SearchDepth */3);
+
+			if (bFirstAnimator)
+			{
+				SupportedProperties = AnimatorSupportedProperties;
+				bFirstAnimator = false;
+			}
+			else
+			{
+				SupportedProperties = SupportedProperties.Intersect(AnimatorSupportedProperties);
+			}
+		}
 	}
 
 	UPropertyAnimatorCorePropertyPreset* Empty = nullptr;
@@ -826,23 +940,24 @@ void UE::PropertyAnimatorCoreEditor::Menu::FillLinkAnimatorSubmenu(UToolMenu* In
 	{
 		const FString MenuName = SupportedProperty.GetPropertyDisplayName();
 		const FText MenuLabel = FText::FromString(MenuName + TEXT(" (") + SupportedProperty.GetLeafPropertyTypeName().ToString() + TEXT(")"));
-
+		const FString PropertyLocatorPath = SupportedProperty.GetPropertyLocatorPath();
+		
 		PropertySection.AddMenuEntry(
 			FName(MenuName)
 			, MenuLabel
 			, LOCTEXT("LinkAnimatorPropertySection.Tooltip", "Link or unlink this property from the animator")
 			, FSlateIcon()
 			, FUIAction(
-				FExecuteAction::CreateLambda(&ExecuteLinkAnimatorPropertyAction, InAnimator, SupportedProperty, Empty, InMenuData)
-				, FCanExecuteAction::CreateLambda(&IsAnimatorLinkPropertyAllowed, InAnimator, SupportedProperty)
-				, FIsActionChecked::CreateLambda(&IsAnimatorPropertyLinked, InAnimator, SupportedProperty)
+				FExecuteAction::CreateLambda(&ExecuteLinkAnimatorPropertyAction, InAnimators, PropertyLocatorPath, Empty, InMenuData)
+				, FCanExecuteAction::CreateLambda(&IsAnimatorLinkPropertyAllowed, InAnimators, PropertyLocatorPath)
+				, FGetActionCheckState::CreateLambda(&GetAnimatorPropertyLinkState, InAnimators, PropertyLocatorPath)
 			)
 			, EUserInterfaceActionType::ToggleButton
 		);
 	}
 }
 
-void UE::PropertyAnimatorCoreEditor::Menu::FillPresetAnimatorSubmenu(UToolMenu* InMenu, UPropertyAnimatorCoreBase* InAnimator, UPropertyAnimatorCorePropertyPreset* InPreset, TSharedRef<FPropertyAnimatorCoreEditorMenuData> InMenuData)
+void UE::PropertyAnimatorCoreEditor::Menu::FillPresetAnimatorSubmenu(UToolMenu* InMenu, const TSet<UPropertyAnimatorCoreBase*>& InAnimators, UPropertyAnimatorCorePropertyPreset* InPreset, TSharedRef<FPropertyAnimatorCoreEditorMenuData> InMenuData)
 {
 	if (!InMenu)
 	{
@@ -855,9 +970,9 @@ void UE::PropertyAnimatorCoreEditor::Menu::FillPresetAnimatorSubmenu(UToolMenu*
 		, LOCTEXT("LinkAllPresetProperty.Tooltip", "Link all properties from this preset")
 		, FSlateIcon()
 		, FUIAction(
-			FExecuteAction::CreateLambda(&ExecuteLinkAnimatorPresetAction, InAnimator, InPreset, InMenuData)
+			FExecuteAction::CreateLambda(&ExecuteLinkAnimatorPresetAction, InAnimators, InPreset, InMenuData)
 			, FCanExecuteAction()
-			, FGetActionCheckState::CreateLambda(&GetAnimatorPresetState, InAnimator, InPreset)
+			, FGetActionCheckState::CreateLambda(&GetAnimatorPresetState, InAnimators, InPreset)
 		)
 		, EUserInterfaceActionType::ToggleButton
 	);
@@ -871,14 +986,36 @@ void UE::PropertyAnimatorCoreEditor::Menu::FillPresetAnimatorSubmenu(UToolMenu*
 	const FToolMenuEntry SeparatorEntry = FToolMenuEntry::InitSeparator(TEXT("PresetSeparator"));
 	InMenu->AddMenuEntry(SeparatorEntry.Name, SeparatorEntry);
 
+	bool bFirstAnimator = true;
 	TSet<FPropertyAnimatorCoreData> SupportedProperties;
-	InPreset->GetSupportedPresetProperties(InAnimator->GetAnimatorActor(), InAnimator, SupportedProperties);
 
+	for (UPropertyAnimatorCoreBase* Animator : InAnimators)
+	{
+		if (!IsValid(Animator) || Animator->IsTemplate())
+		{
+			continue;
+		}
+
+		TSet<FPropertyAnimatorCoreData> AnimatorSupportedProperties;
+		InPreset->GetSupportedPresetProperties(Animator->GetAnimatorActor(), Animator, AnimatorSupportedProperties);
+
+		if (bFirstAnimator)
+		{
+			SupportedProperties = AnimatorSupportedProperties;
+			bFirstAnimator = false;
+		}
+		else
+		{
+			SupportedProperties = SupportedProperties.Intersect(AnimatorSupportedProperties);
+		}
+	}
+	
 	for (const FPropertyAnimatorCoreData& SupportedProperty : SupportedProperties)
 	{
 		const FString MenuName = SupportedProperty.GetPropertyDisplayName();
 		const FText MenuLabel = FText::FromString(MenuName + TEXT(" (") + SupportedProperty.GetLeafPropertyTypeName().ToString() + TEXT(")"));
 		const FText MenuTooltip = LOCTEXT("LinkPresetProperty.Tooltip", "Link this preset property");
+		const FString PropertyLocatorPath = SupportedProperty.GetPropertyLocatorPath();
 
 		const FToolMenuEntry SupportedPropertyEntry = FToolMenuEntry::InitMenuEntry(
 			FName(MenuName)
@@ -886,9 +1023,9 @@ void UE::PropertyAnimatorCoreEditor::Menu::FillPresetAnimatorSubmenu(UToolMenu*
 			, MenuTooltip
 			, FSlateIcon()
 			, FUIAction(
-				FExecuteAction::CreateLambda(&ExecuteLinkAnimatorPropertyAction, InAnimator, SupportedProperty, InPreset, InMenuData)
-				, FCanExecuteAction::CreateLambda(&IsAnimatorLinkPropertyAllowed, InAnimator, SupportedProperty)
-				, FIsActionChecked::CreateLambda(&IsAnimatorPropertyLinked, InAnimator, SupportedProperty)
+				FExecuteAction::CreateLambda(&ExecuteLinkAnimatorPropertyAction, InAnimators, PropertyLocatorPath, InPreset, InMenuData)
+				, FCanExecuteAction::CreateLambda(&IsAnimatorLinkPropertyAllowed, InAnimators, PropertyLocatorPath)
+				, FGetActionCheckState::CreateLambda(&GetAnimatorPropertyLinkState, InAnimators, PropertyLocatorPath)
 			)
 			, EUserInterfaceActionType::ToggleButton
 		);
@@ -969,14 +1106,15 @@ void UE::PropertyAnimatorCoreEditor::Menu::FillNewPresetAnimatorSubmenu(UToolMen
 	{
 		const FString MenuName = PresetProperty.GetPropertyDisplayName();
 		const FText MenuLabel = FText::FromString(MenuName + TEXT(" (") + PresetProperty.GetLeafPropertyTypeName().ToString() + TEXT(")"));
-
+		const FString PropertyLocatorPath = PresetProperty.GetPropertyLocatorPath();
+		
 		const FToolMenuEntry CreatePropertyEntry = FToolMenuEntry::InitMenuEntry(
 			FName(TEXT("CreateProperty") + MenuName)
 			, MenuLabel
 			, LOCTEXT("CreateAnimatorPresetProperty.Tooltip", "Create this animator using this preset property")
 			, FSlateIcon()
 			, FUIAction(
-				FExecuteAction::CreateLambda(&ExecuteNewAnimatorPropertyAction, InAnimator, PresetProperty, InMenuData)
+				FExecuteAction::CreateLambda(&ExecuteNewAnimatorPropertyAction, InAnimator, PropertyLocatorPath, InMenuData)
 				, FCanExecuteAction()
 				, FIsActionChecked()
 				, FIsActionButtonVisible::CreateLambda(&IsLastAnimatorCreatedActionHidden, InAnimator, InMenuData)
@@ -995,9 +1133,9 @@ void UE::PropertyAnimatorCoreEditor::Menu::FillNewPresetAnimatorSubmenu(UToolMen
 			, LOCTEXT("LinkAnimatorPresetProperty.Tooltip", "Link this preset property to the last created animator")
 			, FSlateIcon()
 			, FUIAction(
-				FExecuteAction::CreateLambda(&ExecuteLinkLastCreatedAnimatorPropertyAction, InAnimator, PresetProperty, InMenuData)
+				FExecuteAction::CreateLambda(&ExecuteLinkLastCreatedAnimatorPropertyAction, InAnimator, PropertyLocatorPath, InMenuData)
 				, FCanExecuteAction()
-				, FIsActionChecked::CreateLambda(&IsLastAnimatorCreatedPropertyLinked, InAnimator, PresetProperty, InMenuData)
+				, FIsActionChecked::CreateLambda(&IsLastAnimatorCreatedPropertyLinked, InAnimator, PropertyLocatorPath, InMenuData)
 				, FIsActionButtonVisible::CreateLambda(&IsLastAnimatorCreatedActionVisible, InAnimator, InMenuData)
 			)
 			, EUserInterfaceActionType::ToggleButton
@@ -1021,19 +1159,41 @@ bool UE::PropertyAnimatorCoreEditor::Menu::IsAnimatorPropertyLinked(const UPrope
 	return InAnimator->IsPropertyLinked(InProperty);
 }
 
-bool UE::PropertyAnimatorCoreEditor::Menu::IsAnimatorLinkPropertyAllowed(const UPropertyAnimatorCoreBase* InAnimator, const FPropertyAnimatorCoreData& InProperty)
+bool UE::PropertyAnimatorCoreEditor::Menu::IsAnimatorLinkPropertyAllowed(const TSet<UPropertyAnimatorCoreBase*>& InAnimators, const FString& InPropertyLocatorPath)
 {
-	if (!IsValid(InAnimator) || InAnimator->IsTemplate())
+	if (InAnimators.IsEmpty())
 	{
 		return false;
 	}
 
-	// Only allow linking properties that are not yet linked and do not have any of their children linked
-	return (!InAnimator->IsPropertyLinked(InProperty) && InAnimator->GetInnerPropertiesLinked(InProperty).IsEmpty())
-		|| InAnimator->GetLinkedPropertyContext(InProperty) != nullptr;
+	bool bAllowed = false;
+
+	for (const UPropertyAnimatorCoreBase* Animator : InAnimators)
+	{
+		if (!IsValid(Animator) || Animator->IsTemplate())
+		{
+			continue;
+		}
+
+		FPropertyAnimatorCoreData AnimatorProperty(Animator->GetAnimatorActor(), InPropertyLocatorPath);
+		if (!AnimatorProperty.IsResolved())
+		{
+			continue;
+		}
+
+		// Only allow linking properties that are not yet linked and do not have any of their children linked
+		bAllowed |= (!Animator->IsPropertyLinked(AnimatorProperty) && Animator->GetInnerPropertiesLinked(AnimatorProperty).IsEmpty()) || Animator->GetLinkedPropertyContext(AnimatorProperty) != nullptr;
+
+		if (!bAllowed)
+		{
+			break;
+		}
+	}
+
+	return bAllowed;
 }
 
-bool UE::PropertyAnimatorCoreEditor::Menu::IsLastAnimatorCreatedPropertyLinked(const UPropertyAnimatorCoreBase* InAnimator, const FPropertyAnimatorCoreData& InProperty, TSharedRef<FPropertyAnimatorCoreEditorMenuData> InMenuData)
+bool UE::PropertyAnimatorCoreEditor::Menu::IsLastAnimatorCreatedPropertyLinked(const UPropertyAnimatorCoreBase* InAnimator, const FString& InPropertyLocatorPath, TSharedRef<FPropertyAnimatorCoreEditorMenuData> InMenuData)
 {
 	if (!IsValid(InAnimator)
 		|| !InAnimator->IsTemplate()
@@ -1044,8 +1204,13 @@ bool UE::PropertyAnimatorCoreEditor::Menu::IsLastAnimatorCreatedPropertyLinked(c
 
 	for (const UPropertyAnimatorCoreBase* LastCreatedAnimator : InMenuData->GetLastCreatedAnimators())
 	{
-		if (LastCreatedAnimator->GetClass() != InAnimator->GetClass()
-			|| !LastCreatedAnimator->IsPropertyLinked(InProperty))
+		if (LastCreatedAnimator->GetClass() != InAnimator->GetClass())
+		{
+			return false;
+		}
+
+		const FPropertyAnimatorCoreData Property(LastCreatedAnimator->GetAnimatorActor(), InPropertyLocatorPath);
+		if (!LastCreatedAnimator->IsPropertyLinked(Property))
 		{
 			return false;
 		}
diff --git a/Engine/Plugins/Experimental/PropertyAnimatorCore/Source/PropertyAnimatorCoreEditor/Private/Menus/PropertyAnimatorCoreEditorMenu.h b/Engine/Plugins/Experimental/PropertyAnimatorCore/Source/PropertyAnimatorCoreEditor/Private/Menus/PropertyAnimatorCoreEditorMenu.h
--- a/Engine/Plugins/Experimental/PropertyAnimatorCore/Source/PropertyAnimatorCoreEditor/Private/Menus/PropertyAnimatorCoreEditorMenu.h
+++ b/Engine/Plugins/Experimental/PropertyAnimatorCore/Source/PropertyAnimatorCoreEditor/Private/Menus/PropertyAnimatorCoreEditorMenu.h
@@ -35,9 +35,9 @@ namespace UE::PropertyAnimatorCoreEditor::Menu
 
 	void FillNewAnimatorSubmenu(UToolMenu* InMenu, UPropertyAnimatorCoreBase* InAnimator, TSharedRef<FPropertyAnimatorCoreEditorMenuData> InMenuData);
 
-	void FillLinkAnimatorSubmenu(UToolMenu* InMenu, UPropertyAnimatorCoreBase* InAnimator, TSharedRef<FPropertyAnimatorCoreEditorMenuData> InMenuData);
+	void FillLinkAnimatorSubmenu(UToolMenu* InMenu, const TSet<UPropertyAnimatorCoreBase*>& InAnimators, TSharedRef<FPropertyAnimatorCoreEditorMenuData> InMenuData);
 
-	void FillPresetAnimatorSubmenu(UToolMenu* InMenu, UPropertyAnimatorCoreBase* InAnimator, UPropertyAnimatorCorePropertyPreset* InPreset, TSharedRef<FPropertyAnimatorCoreEditorMenuData> InMenuData);
+	void FillPresetAnimatorSubmenu(UToolMenu* InMenu, const TSet<UPropertyAnimatorCoreBase*>& InAnimators, UPropertyAnimatorCorePropertyPreset* InPreset, TSharedRef<FPropertyAnimatorCoreEditorMenuData> InMenuData);
 
 	void FillNewPresetAnimatorSubmenu(UToolMenu* InMenu, UPropertyAnimatorCoreBase* InAnimator, UPropertyAnimatorCorePropertyPreset* InPreset, TSharedRef<FPropertyAnimatorCoreEditorMenuData> InMenuData);
 
@@ -45,15 +45,15 @@ namespace UE::PropertyAnimatorCoreEditor::Menu
 
 	void ExecuteNewAnimatorPresetAction(const UPropertyAnimatorCoreBase* InAnimator, const TSet<AActor*>& InActors, UPropertyAnimatorCorePresetBase* InPreset, TSharedRef<FPropertyAnimatorCoreEditorMenuData> InMenuData);
 
-	void ExecuteNewAnimatorPropertyAction(const UPropertyAnimatorCoreBase* InAnimator, FPropertyAnimatorCoreData InProperty, TSharedRef<FPropertyAnimatorCoreEditorMenuData> InMenuData);
+	void ExecuteNewAnimatorPropertyAction(const UPropertyAnimatorCoreBase* InAnimator, const FString& InPropertyLocatorPath, TSharedRef<FPropertyAnimatorCoreEditorMenuData> InMenuData);
 
-	void ExecuteLinkLastCreatedAnimatorPropertyAction(const UPropertyAnimatorCoreBase* InAnimator, FPropertyAnimatorCoreData InProperty, TSharedRef<FPropertyAnimatorCoreEditorMenuData> InMenuData);
+	void ExecuteLinkLastCreatedAnimatorPropertyAction(const UPropertyAnimatorCoreBase* InAnimator, const FString& InPropertyLocatorPath, TSharedRef<FPropertyAnimatorCoreEditorMenuData> InMenuData);
 
 	void ExecuteApplyLastCreatedAnimatorPresetAction(const UPropertyAnimatorCoreBase* InAnimator, UPropertyAnimatorCorePropertyPreset* InPreset, TSharedRef<FPropertyAnimatorCoreEditorMenuData> InMenuData);
 
-	void ExecuteLinkAnimatorPresetAction(UPropertyAnimatorCoreBase* InAnimator, UPropertyAnimatorCorePropertyPreset* InPreset, TSharedRef<FPropertyAnimatorCoreEditorMenuData> InMenuData);
+	void ExecuteLinkAnimatorPresetAction(const TSet<UPropertyAnimatorCoreBase*>& InAnimators, UPropertyAnimatorCorePropertyPreset* InPreset, TSharedRef<FPropertyAnimatorCoreEditorMenuData> InMenuData);
 
-	void ExecuteLinkAnimatorPropertyAction(UPropertyAnimatorCoreBase* InAnimator, FPropertyAnimatorCoreData InProperty, UPropertyAnimatorCorePropertyPreset* InPreset, TSharedRef<FPropertyAnimatorCoreEditorMenuData> InMenuData);
+	void ExecuteLinkAnimatorPropertyAction(const TSet<UPropertyAnimatorCoreBase*>& InAnimators, const FString& InPropertyLocatorPath, UPropertyAnimatorCorePropertyPreset* InPreset, TSharedRef<FPropertyAnimatorCoreEditorMenuData> InMenuData);
 
 	void ExecuteEnableActorAnimatorAction(TSharedRef<FPropertyAnimatorCoreEditorMenuData> InMenuData, bool bInEnable);
 
@@ -67,17 +67,19 @@ namespace UE::PropertyAnimatorCoreEditor::Menu
 
 	/** Check */
 
-	ECheckBoxState GetAnimatorPresetState(const UPropertyAnimatorCoreBase* InAnimator, UPropertyAnimatorCorePropertyPreset* InPreset);
+	ECheckBoxState GetAnimatorPresetState(const TSet<UPropertyAnimatorCoreBase*>& InAnimators, UPropertyAnimatorCorePropertyPreset* InPreset);
 
 	ECheckBoxState GetLastAnimatorCreatedPresetState(const UPropertyAnimatorCoreBase* InAnimator, UPropertyAnimatorCorePropertyPreset* InPreset, TSharedRef<FPropertyAnimatorCoreEditorMenuData> InMenuData);
 
+	ECheckBoxState GetAnimatorPropertyLinkState(const TSet<UPropertyAnimatorCoreBase*>& InAnimators, const FString& InPropertyLocatorPath);
+
 	bool IsAnimatorPresetLinked(UPropertyAnimatorCoreBase* InAnimator, UPropertyAnimatorCorePropertyPreset* InPreset);
 
 	bool IsAnimatorPropertyLinked(const UPropertyAnimatorCoreBase* InAnimator, const FPropertyAnimatorCoreData& InProperty);
 
-	bool IsAnimatorLinkPropertyAllowed(const UPropertyAnimatorCoreBase* InAnimator, const FPropertyAnimatorCoreData& InProperty);
+	bool IsAnimatorLinkPropertyAllowed(const TSet<UPropertyAnimatorCoreBase*>& InAnimators, const FString& InPropertyLocatorPath);
 
-	bool IsLastAnimatorCreatedPropertyLinked(const UPropertyAnimatorCoreBase* InAnimator, const FPropertyAnimatorCoreData& InProperty, TSharedRef<FPropertyAnimatorCoreEditorMenuData> InMenuData);
+	bool IsLastAnimatorCreatedPropertyLinked(const UPropertyAnimatorCoreBase* InAnimator, const FString& InPropertyLocatorPath, TSharedRef<FPropertyAnimatorCoreEditorMenuData> InMenuData);
 
 	bool IsLastAnimatorCreatedPresetLinked(const UPropertyAnimatorCoreBase* InAnimator, UPropertyAnimatorCorePropertyPreset* InPreset, TSharedRef<FPropertyAnimatorCoreEditorMenuData> InMenuData);
 
diff --git a/Engine/Plugins/Experimental/PropertyAnimatorCore/Source/PropertyAnimatorCoreEditor/Private/Menus/PropertyAnimatorCoreEditorMenuContext.h b/Engine/Plugins/Experimental/PropertyAnimatorCore/Source/PropertyAnimatorCoreEditor/Private/Menus/PropertyAnimatorCoreEditorMenuContext.h
--- a/Engine/Plugins/Experimental/PropertyAnimatorCore/Source/PropertyAnimatorCoreEditor/Private/Menus/PropertyAnimatorCoreEditorMenuContext.h
+++ b/Engine/Plugins/Experimental/PropertyAnimatorCore/Source/PropertyAnimatorCoreEditor/Private/Menus/PropertyAnimatorCoreEditorMenuContext.h
@@ -7,6 +7,9 @@
 #include "UObject/Object.h"
 #include "PropertyAnimatorCoreEditorMenuContext.generated.h"
 
+class UPropertyAnimatorCoreBase;
+struct FPropertyAnimatorCoreData;
+
 /** Context passed in UToolMenu when generating entries with selected items */
 UCLASS()
 class UPropertyAnimatorCoreEditorMenuContext : public UObject
@@ -14,17 +17,28 @@ class UPropertyAnimatorCoreEditorMenuContext : public UObject
 	GENERATED_BODY()
 
 public:
-	const FPropertyAnimatorCoreData& GetPropertyData() const
+	const TArray<FPropertyAnimatorCoreData>& GetProperties() const
+	{
+		return Properties;
+	}
+
+	void SetProperties(const TArray<FPropertyAnimatorCoreData>& InProperties)
+	{
+		Properties = InProperties;
+	}
+
+	const TSet<UPropertyAnimatorCoreBase*>& GetAnimators() const
 	{
-		return PropertyData;
+		return Animators;
 	}
 
-	void SetPropertyData(const FPropertyAnimatorCoreData& InPropertyData)
+	void SetAnimators(const TSet<UPropertyAnimatorCoreBase*>& InAnimators)
 	{
-		PropertyData = InPropertyData;
+		Animators = InAnimators;
 	}
 
 protected:
-	/** The item this menu should apply to */
-	FPropertyAnimatorCoreData PropertyData;
+	/** The items this menu should apply to */
+	TArray<FPropertyAnimatorCoreData> Properties;
+	TSet<UPropertyAnimatorCoreBase*> Animators;
 };
\ No newline at end of file
diff --git a/Engine/Plugins/Experimental/PropertyAnimatorCore/Source/PropertyAnimatorCoreEditor/Private/Menus/PropertyAnimatorCoreEditorMenuDefs.cpp b/Engine/Plugins/Experimental/PropertyAnimatorCore/Source/PropertyAnimatorCoreEditor/Private/Menus/PropertyAnimatorCoreEditorMenuDefs.cpp
--- a/Engine/Plugins/Experimental/PropertyAnimatorCore/Source/PropertyAnimatorCoreEditor/Private/Menus/PropertyAnimatorCoreEditorMenuDefs.cpp
+++ b/Engine/Plugins/Experimental/PropertyAnimatorCore/Source/PropertyAnimatorCoreEditor/Private/Menus/PropertyAnimatorCoreEditorMenuDefs.cpp
@@ -5,7 +5,7 @@
 #include "Animators/PropertyAnimatorCoreBase.h"
 #include "Components/PropertyAnimatorCoreComponent.h"
 
-FPropertyAnimatorCoreEditorMenuContext::FPropertyAnimatorCoreEditorMenuContext(const TSet<UObject*>& InObjects, const TSet<FPropertyAnimatorCoreData>& InProperties)
+FPropertyAnimatorCoreEditorMenuContext::FPropertyAnimatorCoreEditorMenuContext(const TSet<UObject*>& InObjects, const TArray<FPropertyAnimatorCoreData>& InProperties)
 {
 	auto AddActor = [this](AActor* InActor)
 	{
diff --git a/Engine/Plugins/Experimental/PropertyAnimatorCore/Source/PropertyAnimatorCoreEditor/Private/Subsystems/PropertyAnimatorCoreEditorSubsystem.cpp b/Engine/Plugins/Experimental/PropertyAnimatorCore/Source/PropertyAnimatorCoreEditor/Private/Subsystems/PropertyAnimatorCoreEditorSubsystem.cpp
--- a/Engine/Plugins/Experimental/PropertyAnimatorCore/Source/PropertyAnimatorCoreEditor/Private/Subsystems/PropertyAnimatorCoreEditorSubsystem.cpp
+++ b/Engine/Plugins/Experimental/PropertyAnimatorCore/Source/PropertyAnimatorCoreEditor/Private/Subsystems/PropertyAnimatorCoreEditorSubsystem.cpp
@@ -288,8 +288,9 @@ void UPropertyAnimatorCoreEditorSubsystem::OnGetGlobalRowExtension(const FOnGene
 
 void UPropertyAnimatorCoreEditorSubsystem::OnControlPropertyClicked(TWeakPtr<IDetailTreeNode> InOwnerTreeNode, TSharedPtr<IPropertyHandle> InPropertyHandle)
 {
-	TOptional<FPropertyAnimatorCoreData> PropertyData = GetPropertyData(InPropertyHandle);
-	if (!PropertyData.IsSet())
+	constexpr bool bFindMemberProperty = true;
+	TArray<FPropertyAnimatorCoreData> Properties;
+	if (!GetPropertiesFromHandle(InPropertyHandle, Properties, bFindMemberProperty))
 	{
 		return;
 	}
@@ -306,7 +307,7 @@ void UPropertyAnimatorCoreEditorSubsystem::OnControlPropertyClicked(TWeakPtr<IDe
 	FSlateApplication::Get().PushMenu(
 		ActiveWindow.ToSharedRef(),
 		FWidgetPath(),
-		GenerateContextMenuWidget(PropertyData.GetValue()),
+		GenerateContextMenuWidget(Properties),
 		MenuLocation,
 		FPopupTransitionEffect(FPopupTransitionEffect::ContextMenu)
 	);
@@ -314,8 +315,9 @@ void UPropertyAnimatorCoreEditorSubsystem::OnControlPropertyClicked(TWeakPtr<IDe
 
 bool UPropertyAnimatorCoreEditorSubsystem::IsControlPropertySupported(TWeakPtr<IDetailTreeNode> InOwnerTreeNode, TSharedPtr<IPropertyHandle> InPropertyHandle) const
 {
-	TOptional<FPropertyAnimatorCoreData> PropertyData = GetPropertyData(InPropertyHandle);
-	if (!PropertyData.IsSet())
+	constexpr bool bFindMemberProperty = true;
+	TArray<FPropertyAnimatorCoreData> Properties;
+	if (!GetPropertiesFromHandle(InPropertyHandle, Properties, bFindMemberProperty))
 	{
 		return false;
 	}
@@ -326,7 +328,15 @@ bool UPropertyAnimatorCoreEditorSubsystem::IsControlPropertySupported(TWeakPtr<I
 		return false;
 	}
 
-	return Subsystem->IsPropertySupported(PropertyData.GetValue());
+	for (FPropertyAnimatorCoreData& Property : Properties)
+	{
+		if (!Subsystem->IsPropertySupported(Property))
+		{
+			return false;
+		}
+	}
+
+	return true;
 }
 
 bool UPropertyAnimatorCoreEditorSubsystem::IsControlPropertyVisible(TWeakPtr<IDetailTreeNode> InOwnerTreeNode, TSharedPtr<IPropertyHandle> InPropertyHandle) const
@@ -336,8 +346,9 @@ bool UPropertyAnimatorCoreEditorSubsystem::IsControlPropertyVisible(TWeakPtr<IDe
 
 bool UPropertyAnimatorCoreEditorSubsystem::IsControlPropertyLinked(TWeakPtr<IDetailTreeNode> InOwnerTreeNode, TSharedPtr<IPropertyHandle> InPropertyHandle) const
 {
-	TOptional<FPropertyAnimatorCoreData> PropertyData = GetPropertyData(InPropertyHandle);
-	if (!PropertyData.IsSet())
+	constexpr bool bFindMemberProperty = true;
+	TArray<FPropertyAnimatorCoreData> Properties;
+	if (!GetPropertiesFromHandle(InPropertyHandle, Properties, bFindMemberProperty))
 	{
 		return false;
 	}
@@ -348,23 +359,36 @@ bool UPropertyAnimatorCoreEditorSubsystem::IsControlPropertyLinked(TWeakPtr<IDet
 		return false;
 	}
 
-	for (const UPropertyAnimatorCoreBase* Controller : Subsystem->GetExistingAnimators(PropertyData.GetValue().GetOwningActor()))
+	bool bAnimated = false;
+	for (FPropertyAnimatorCoreData& Property : Properties)
 	{
-		if (Controller->IsPropertyLinked(PropertyData.GetValue()))
+		for (const UPropertyAnimatorCoreBase* Animator : Subsystem->GetExistingAnimators(Property.GetOwningActor()))
 		{
-			return true;
+			if (Animator->IsPropertyLinked(Property))
+			{
+				bAnimated |= true;
+			}
+			else if (!Animator->GetInnerPropertiesLinked(Property).IsEmpty())
+			{
+				bAnimated |= true;
+			}
+			else
+			{
+				bAnimated |= false;
+				break;
+			}
 		}
 
-		if (!Controller->GetInnerPropertiesLinked(PropertyData.GetValue()).IsEmpty())
+		if (!bAnimated)
 		{
-			return true;
+			break;
 		}
 	}
 
-	return false;
+	return bAnimated;
 }
 
-TSharedRef<SWidget> UPropertyAnimatorCoreEditorSubsystem::GenerateContextMenuWidget(const FPropertyAnimatorCoreData& InProperty)
+TSharedRef<SWidget> UPropertyAnimatorCoreEditorSubsystem::GenerateContextMenuWidget(const TArray<FPropertyAnimatorCoreData>& InProperties)
 {
 	UToolMenus* Menus = UToolMenus::Get();
 
@@ -383,7 +407,7 @@ TSharedRef<SWidget> UPropertyAnimatorCoreEditorSubsystem::GenerateContextMenuWid
 	}
 
 	UPropertyAnimatorCoreEditorMenuContext* MenuContext = NewObject<UPropertyAnimatorCoreEditorMenuContext>();
-	MenuContext->SetPropertyData(InProperty);
+	MenuContext->SetProperties(InProperties);
 
 	const FToolMenuContext ToolMenuContext(MenuContext);
 	return Menus->GenerateWidget(AnimatorExtensionMenuName, ToolMenuContext);
@@ -417,7 +441,7 @@ void UPropertyAnimatorCoreEditorSubsystem::FillAnimatorExtensionSection(UToolMen
 		return;
 	}
 
-	const FPropertyAnimatorCoreEditorMenuContext MenuContext({}, {Context->GetPropertyData()});
+	const FPropertyAnimatorCoreEditorMenuContext MenuContext({}, Context->GetProperties());
 	const FPropertyAnimatorCoreEditorMenuOptions MenuOptions(
 		{
 			EPropertyAnimatorCoreEditorMenuType::NewAdvanced
@@ -442,14 +466,14 @@ void UPropertyAnimatorCoreEditorSubsystem::FillAnimatorRowContextSection(UToolMe
 		return;
 	}
 
-	TOptional<FPropertyAnimatorCoreData> PropertyData = GetPropertyData(Context->PropertyHandles[0], true);
-
-	if (!PropertyData.IsSet())
+	constexpr bool bFindMemberProperty = true;
+	TArray<FPropertyAnimatorCoreData> Properties;
+	if (!GetPropertiesFromHandle(Context->PropertyHandles[0], Properties, bFindMemberProperty))
 	{
 		return;
 	}
 
-	const FPropertyAnimatorCoreEditorMenuContext MenuContext({}, {PropertyData.GetValue()});
+	const FPropertyAnimatorCoreEditorMenuContext MenuContext({}, Properties);
 	const FPropertyAnimatorCoreEditorMenuOptions MenuOptions(
 		{
 			EPropertyAnimatorCoreEditorMenuType::NewAdvanced
@@ -459,48 +483,42 @@ void UPropertyAnimatorCoreEditorSubsystem::FillAnimatorRowContextSection(UToolMe
 	FillAnimatorMenu(InToolMenu, MenuContext, MenuOptions);
 }
 
-TOptional<FPropertyAnimatorCoreData> UPropertyAnimatorCoreEditorSubsystem::GetPropertyData(TSharedPtr<IPropertyHandle> InPropertyHandle, bool bInFindMemberProperty) const
+bool UPropertyAnimatorCoreEditorSubsystem::GetPropertiesFromHandle(const TSharedPtr<IPropertyHandle>& InPropertyHandle, TArray<FPropertyAnimatorCoreData>& OutProperties, bool bInFindMemberProperty) const
 {
 	if (!InPropertyHandle.IsValid() || !InPropertyHandle->IsValidHandle())
 	{
-		return TOptional<FPropertyAnimatorCoreData>();
+		return false;
 	}
 
 	const FProperty* Property = InPropertyHandle->GetProperty();
 	if (!Property)
 	{
-		return TOptional<FPropertyAnimatorCoreData>();
+		return false;
 	}
 
-	TArray<UObject*> OutOwners;
-	InPropertyHandle->GetOuterObjects(OutOwners);
-
-	if (OutOwners.Num() != 1)
-	{
-		return TOptional<FPropertyAnimatorCoreData>();
-	}
+	TArray<UObject*> Owners;
+	InPropertyHandle->GetOuterObjects(Owners);
 
-	UObject* Owner = OutOwners[0];
-	if (!IsValid(Owner))
+	if (Owners.IsEmpty() || !IsValid(Owners[0]))
 	{
-		return TOptional<FPropertyAnimatorCoreData>();
+		return false;
 	}
 
-	if (Owner->IsA<UPropertyAnimatorCoreComponent>() || Owner->GetTypedOuter<UPropertyAnimatorCoreComponent>())
+	if (Owners[0]->IsA<UPropertyAnimatorCoreComponent>() || Owners[0]->GetTypedOuter<UPropertyAnimatorCoreComponent>())
 	{
-		return TOptional<FPropertyAnimatorCoreData>();
+		return false;
 	}
 
-	const FString OwnerPath = Owner->GetPathName();
+	const FString OwnerPath = Owners[0]->GetPathName();
 	if (OwnerPath.IsEmpty())
 	{
-		return TOptional<FPropertyAnimatorCoreData>();
+		return false;
 	}
 
 	const FString PropertyPath = InPropertyHandle->GeneratePathToProperty();
 	if (PropertyPath.IsEmpty())
 	{
-		return TOptional<FPropertyAnimatorCoreData>();
+		return false;
 	}
 
 	// Climbs up the tree to find a member property handle that is a direct child of an object property
@@ -524,18 +542,27 @@ TOptional<FPropertyAnimatorCoreData> UPropertyAnimatorCoreEditorSubsystem::GetPr
 
 	if (!MemberPropertyHandle.IsValid() || !MemberPropertyHandle->IsValidHandle())
 	{
-		return TOptional<FPropertyAnimatorCoreData>();
+		return false;
 	}
 
-	const FPropertyAnimatorCoreData PropertyData(Owner, MemberPropertyHandle->GetProperty(), nullptr);
+	FPropertyAnimatorCoreData PropertyData(Owners[0], MemberPropertyHandle->GetProperty(), nullptr);
 
 	// We need a setter to control a property
 	if (!PropertyData.HasSetter())
 	{
-		return TOptional<FPropertyAnimatorCoreData>();
+		return false;
+	}
+
+	OutProperties.Emplace(PropertyData);
+	for (int32 Index = 1; Index < Owners.Num(); Index++)
+	{
+		if (IsValid(Owners[Index]))
+		{
+			OutProperties.Emplace(FPropertyAnimatorCoreData(Owners[Index], MemberPropertyHandle->GetProperty(), nullptr));
+		}
 	}
 
-	return PropertyData;
+	return true;
 }
 
 #undef LOCTEXT_NAMESPACE
diff --git a/Engine/Plugins/Experimental/PropertyAnimatorCore/Source/PropertyAnimatorCoreEditor/Public/Menus/PropertyAnimatorCoreEditorMenuDefs.h b/Engine/Plugins/Experimental/PropertyAnimatorCore/Source/PropertyAnimatorCoreEditor/Public/Menus/PropertyAnimatorCoreEditorMenuDefs.h
--- a/Engine/Plugins/Experimental/PropertyAnimatorCore/Source/PropertyAnimatorCoreEditor/Public/Menus/PropertyAnimatorCoreEditorMenuDefs.h
+++ b/Engine/Plugins/Experimental/PropertyAnimatorCore/Source/PropertyAnimatorCoreEditor/Public/Menus/PropertyAnimatorCoreEditorMenuDefs.h
@@ -35,7 +35,7 @@ enum class EPropertyAnimatorCoreEditorMenuType : uint8
 struct PROPERTYANIMATORCOREEDITOR_API FPropertyAnimatorCoreEditorMenuContext
 {
 	FPropertyAnimatorCoreEditorMenuContext() {}
-	explicit FPropertyAnimatorCoreEditorMenuContext(const TSet<UObject*>& InObjects, const TSet<FPropertyAnimatorCoreData>& InProperties);
+	explicit FPropertyAnimatorCoreEditorMenuContext(const TSet<UObject*>& InObjects, const TArray<FPropertyAnimatorCoreData>& InProperties);
 
 	const TSet<AActor*>& GetActors() const;
 	const TSet<FPropertyAnimatorCoreData>& GetProperties() const;
diff --git a/Engine/Plugins/Experimental/PropertyAnimatorCore/Source/PropertyAnimatorCoreEditor/Public/Subsystems/PropertyAnimatorCoreEditorSubsystem.h b/Engine/Plugins/Experimental/PropertyAnimatorCore/Source/PropertyAnimatorCoreEditor/Public/Subsystems/PropertyAnimatorCoreEditorSubsystem.h
--- a/Engine/Plugins/Experimental/PropertyAnimatorCore/Source/PropertyAnimatorCoreEditor/Public/Subsystems/PropertyAnimatorCoreEditorSubsystem.h
+++ b/Engine/Plugins/Experimental/PropertyAnimatorCore/Source/PropertyAnimatorCoreEditor/Public/Subsystems/PropertyAnimatorCoreEditorSubsystem.h
@@ -66,7 +66,7 @@ protected:
 	bool IsControlPropertyLinked(TWeakPtr<IDetailTreeNode> InOwnerTreeNode, TSharedPtr<IPropertyHandle> InPropertyHandle) const;
 
 	/** Creates the context menu to display when selecting the animator icon */
-	TSharedRef<SWidget> GenerateContextMenuWidget(const FPropertyAnimatorCoreData& InProperty);
+	TSharedRef<SWidget> GenerateContextMenuWidget(const TArray<FPropertyAnimatorCoreData>& InProperties);
 
 	/** Extend the context menu on each property row in details panel with additional entries */
 	void ExtendPropertyRowContextMenu();
@@ -77,8 +77,8 @@ protected:
 	/** Fills the animator details view row context menu */
 	void FillAnimatorRowContextSection(UToolMenu* InToolMenu);
 
-	/** Convert a property handle to a property data */
-	TOptional<FPropertyAnimatorCoreData> GetPropertyData(TSharedPtr<IPropertyHandle> InPropertyHandle, bool bInFindMemberProperty = false) const;
+	/** Extracts property data along with owners from a property handle */
+	bool GetPropertiesFromHandle(const TSharedPtr<IPropertyHandle>& InPropertyHandle, TArray<FPropertyAnimatorCoreData>& OutProperties, bool bInFindMemberProperty = false) const;
 
 	FDelegateHandle OnGetGlobalRowExtensionHandle;
 
