MotionDesign : Animators
- Allowing multi-edit for property animator properties options (amplitude, magnitude, ...) as long as linked properties are the same and count matches

diff --git a/Engine/Plugins/Experimental/PropertyAnimatorCore/Source/PropertyAnimatorCoreEditor/Private/Customizations/PropertyAnimatorCoreEditorContextTypeCustomization.cpp b/Engine/Plugins/Experimental/PropertyAnimatorCore/Source/PropertyAnimatorCoreEditor/Private/Customizations/PropertyAnimatorCoreEditorContextTypeCustomization.cpp
--- a/Engine/Plugins/Experimental/PropertyAnimatorCore/Source/PropertyAnimatorCoreEditor/Private/Customizations/PropertyAnimatorCoreEditorContextTypeCustomization.cpp
+++ b/Engine/Plugins/Experimental/PropertyAnimatorCore/Source/PropertyAnimatorCoreEditor/Private/Customizations/PropertyAnimatorCoreEditorContextTypeCustomization.cpp
@@ -4,7 +4,6 @@
 
 #include "DetailWidgetRow.h"
 #include "IDetailChildrenBuilder.h"
-#include "IPropertyUtilities.h"
 #include "Subsystems/PropertyAnimatorCoreSubsystem.h"
 #include "Widgets/Images/SImage.h"
 #include "Widgets/Input/SButton.h"
@@ -20,20 +19,34 @@ void FPropertyAnimatorCoreEditorContextTypeCustomization::CustomizeHeader(TShare
 		return;
 	}
 
-	if (InPropertyHandle->GetNumOuterObjects() != 1)
+	// Dive into the array handle
+	const TSharedPtr<IPropertyHandle> LinkedPropertiesHandle = InPropertyHandle->GetChildHandle(0);
+	if (!LinkedPropertiesHandle.IsValid() || !LinkedPropertiesHandle->IsValidHandle())
 	{
 		return;
 	}
 
-	const UPropertyAnimatorCoreContext* PropertyContext = GetPropertyContextValue(InPropertyHandle);
+	// Dive into array slot
+	const TSharedPtr<IPropertyHandle> LinkedPropertyHandle = LinkedPropertiesHandle->GetChildHandle(0);
+	if (!LinkedPropertyHandle.IsValid() || !LinkedPropertyHandle->IsValidHandle())
+	{
+		return;
+	}
 
-	if (!PropertyContext)
+	TArray<UObject*> Objects;
+	LinkedPropertyHandle->GetOuterObjects(Objects);
+	if (Objects.IsEmpty())
 	{
 		return;
 	}
 
-	PropertyContextHandle = InPropertyHandle;
+	const UPropertyAnimatorCoreContext* PropertyContext = Cast<UPropertyAnimatorCoreContext>(Objects[0]);
+	if (!PropertyContext)
+	{
+		return;
+	}
 
+	PropertyContextHandle = LinkedPropertyHandle;
 	const FString PropertyDisplayName = PropertyContext->GetAnimatedProperty().GetPropertyDisplayName();
 	const FName PropertyTypeName = PropertyContext->GetAnimatedProperty().GetLeafPropertyTypeName();
 
@@ -51,7 +64,7 @@ void FPropertyAnimatorCoreEditorContextTypeCustomization::CustomizeHeader(TShare
 			+ SHorizontalBox::Slot()
 			.FillWidth(1.f)
 			[
-				InPropertyHandle->CreatePropertyNameWidget(FText::FromString(PropertyDisplayName + TEXT(" (") + PropertyTypeName.ToString() + TEXT(")")))
+				LinkedPropertiesHandle->CreatePropertyNameWidget(FText::FromString(PropertyDisplayName + TEXT(" (") + PropertyTypeName.ToString() + TEXT(")")))
 			]
 		]
 		.ValueContent()
@@ -78,95 +91,75 @@ void FPropertyAnimatorCoreEditorContextTypeCustomization::CustomizeHeader(TShare
 
 void FPropertyAnimatorCoreEditorContextTypeCustomization::CustomizeChildren(TSharedRef<IPropertyHandle> InPropertyHandle, IDetailChildrenBuilder& InBuilder, IPropertyTypeCustomizationUtils& InUtils)
 {
-	UObject* PropertyObject;
-	if (FPropertyAccess::Success != InPropertyHandle->GetValue(PropertyObject))
-	{
-		return;
-	}
-
 	static const TSet<FName> SkipProperties
 	{
 		GET_MEMBER_NAME_CHECKED(UPropertyAnimatorCoreContext, bAnimated)
 	};
 
-	if (UPropertyAnimatorCoreContext* Options = Cast<UPropertyAnimatorCoreContext>(PropertyObject))
+	if (PropertyContextHandle.IsValid() && PropertyContextHandle->IsValidHandle())
 	{
-		for (const FProperty* Property : TFieldRange<FProperty>(Options->GetClass()))
+		uint32 ChildrenCount = 0;
+		PropertyContextHandle->GetNumChildren(ChildrenCount);
+
+		for (uint32 ChildIndex = 0; ChildIndex < ChildrenCount; ++ChildIndex)
 		{
-			if (Property
-				&& Property->HasAnyPropertyFlags(EPropertyFlags::CPF_Edit)
-				&& !SkipProperties.Contains(Property->GetFName()))
-			{
-				IDetailPropertyRow* Row = InBuilder.AddExternalObjectProperty({Options}, Property->GetFName());
+			TSharedPtr<IPropertyHandle> ChildHandle = PropertyContextHandle->GetChildHandle(ChildIndex);
 
-				if (Row && Property->GetFName() == GET_MEMBER_NAME_CHECKED(UPropertyAnimatorCoreContext, ConverterRule))
-				{
-					if (const TSharedPtr<IPropertyHandle> ConverterHandle = Row->GetPropertyHandle())
-					{
-						PropertyUtilities = InUtils.GetPropertyUtilities();
-
-						// Fix for container properties (array, set, map) in instanced struct not updating details view (when adding/removing entries)
-						ConverterHandle->SetOnChildPropertyValueChangedWithData(TDelegate<void(const FPropertyChangedEvent&)>::CreateLambda([this, Row](const FPropertyChangedEvent& InEvent)
-						{
-							if (PropertyUtilities.IsValid()
-								&& InEvent.Property
-								&& (InEvent.Property->IsA<FArrayProperty>()
-								|| InEvent.Property->IsA<FSetProperty>()
-								|| InEvent.Property->IsA<FMapProperty>()))
-							{
-								PropertyUtilities->RequestRefresh();
-							}
-						}));
-
-						// Hide customization dropdown to avoid user changing instanced struct type
-						PropertyUtilities->EnqueueDeferredAction(FSimpleDelegate::CreateSP(this, &FPropertyAnimatorCoreEditorContextTypeCustomization::HideCustomValueWidget, Row));
-					}
-				}
+			if (ChildHandle.IsValid() && !SkipProperties.Contains(ChildHandle->GetProperty()->GetFName()))
+			{
+				InBuilder.AddProperty(ChildHandle.ToSharedRef());
 			}
 		}
 	}
 }
 
-UPropertyAnimatorCoreContext* FPropertyAnimatorCoreEditorContextTypeCustomization::GetPropertyContextValue(TSharedPtr<IPropertyHandle> InHandle)
+ECheckBoxState FPropertyAnimatorCoreEditorContextTypeCustomization::IsPropertyEnabled() const
 {
-	if (!InHandle.IsValid() || !InHandle->IsValidHandle())
-	{
-		return nullptr;
-	}
-
-	UObject* ValueData = nullptr;
+	TOptional<ECheckBoxState> State;
 
-	if (InHandle->GetValue(ValueData) == FPropertyAccess::Result::Success)
+	if (PropertyContextHandle.IsValid())
 	{
-		return Cast<UPropertyAnimatorCoreContext>(ValueData);
-	}
+		TArray<UObject*> Objects;
+		PropertyContextHandle->GetOuterObjects(Objects);
 
-	return nullptr;
-}
+		for (UObject* Object : Objects)
+		{
+			if (UPropertyAnimatorCoreContext* ContextObject = Cast<UPropertyAnimatorCoreContext>(Object))
+			{
+				const ECheckBoxState ContextState = ContextObject->IsAnimated() ? ECheckBoxState::Checked : ECheckBoxState::Unchecked;
 
-ECheckBoxState FPropertyAnimatorCoreEditorContextTypeCustomization::IsPropertyEnabled() const
-{
-	if (const UPropertyAnimatorCoreContext* PropertyContext = GetPropertyContextValue(PropertyContextHandle))
-	{
-		return PropertyContext->IsAnimated() ? ECheckBoxState::Checked : ECheckBoxState::Unchecked;
+				if (!State.IsSet())
+				{
+					State = ContextState;
+				}
+				else if (State.GetValue() != ContextState)
+				{
+					State = ECheckBoxState::Undetermined;
+					break;
+				}
+			}
+		}
 	}
 
-	return ECheckBoxState::Undetermined;
+	return State.Get(ECheckBoxState::Undetermined);
 }
 
 void FPropertyAnimatorCoreEditorContextTypeCustomization::OnPropertyEnabled(ECheckBoxState InNewState) const
 {
-	if (UPropertyAnimatorCoreContext* PropertyContext = GetPropertyContextValue(PropertyContextHandle))
+	if (!PropertyContextHandle.IsValid())
 	{
-		PropertyContext->SetAnimated(InNewState == ECheckBoxState::Checked);
+		return;
 	}
-}
 
-void FPropertyAnimatorCoreEditorContextTypeCustomization::HideCustomValueWidget(IDetailPropertyRow* InConverterRow)
-{
-	if (InConverterRow && InConverterRow->CustomValueWidget())
+	TArray<UObject*> Objects;
+	PropertyContextHandle->GetOuterObjects(Objects);
+
+	for (UObject* Object : Objects)
 	{
-		InConverterRow->CustomValueWidget()->Widget->SetVisibility(EVisibility::Hidden);
+		if (UPropertyAnimatorCoreContext* ContextObject = Cast<UPropertyAnimatorCoreContext>(Object))
+		{
+			ContextObject->SetAnimated(InNewState == ECheckBoxState::Checked);
+		}
 	}
 }
 
@@ -177,24 +170,19 @@ FReply FPropertyAnimatorCoreEditorContextTypeCustomization::UnlinkProperty() con
 		return FReply::Unhandled();
 	}
 
-	const UPropertyAnimatorCoreContext* PropertyContext = GetPropertyContextValue(PropertyContextHandle);
-	if (!PropertyContext)
-	{
-		return FReply::Unhandled();
-	}
+	UPropertyAnimatorCoreSubsystem* AnimatorSubsystem = UPropertyAnimatorCoreSubsystem::Get();
+	TArray<UObject*> Objects;
+	PropertyContextHandle->GetOuterObjects(Objects);
 
-	TArray<UObject*> Owners;
-	PropertyContextHandle->GetOuterObjects(Owners);
-
-	if (UPropertyAnimatorCoreSubsystem* AnimatorSubsystem = UPropertyAnimatorCoreSubsystem::Get())
+	constexpr bool bShouldTransact = true;
+	for (UObject* Object : Objects)
 	{
-		FPropertyAnimatorCoreData PropertyData = PropertyContext->GetAnimatedProperty();
-
-		for (UObject* Owner : Owners)
+		if (const UPropertyAnimatorCoreContext* ContextObject = Cast<UPropertyAnimatorCoreContext>(Object))
 		{
-			if (UPropertyAnimatorCoreBase* PropertyAnimator = Cast<UPropertyAnimatorCoreBase>(Owner))
+			if (UPropertyAnimatorCoreBase* PropertyAnimator = ContextObject->GetTypedOuter<UPropertyAnimatorCoreBase>())
 			{
-				AnimatorSubsystem->UnlinkAnimatorProperty(PropertyAnimator, PropertyData, true);
+				FPropertyAnimatorCoreData Property = ContextObject->GetAnimatedProperty();
+				AnimatorSubsystem->UnlinkAnimatorProperty(PropertyAnimator, Property, bShouldTransact);
 			}
 		}
 	}
diff --git a/Engine/Plugins/Experimental/PropertyAnimatorCore/Source/PropertyAnimatorCoreEditor/Private/Customizations/PropertyAnimatorCoreEditorContextTypeCustomization.h b/Engine/Plugins/Experimental/PropertyAnimatorCore/Source/PropertyAnimatorCoreEditor/Private/Customizations/PropertyAnimatorCoreEditorContextTypeCustomization.h
--- a/Engine/Plugins/Experimental/PropertyAnimatorCore/Source/PropertyAnimatorCoreEditor/Private/Customizations/PropertyAnimatorCoreEditorContextTypeCustomization.h
+++ b/Engine/Plugins/Experimental/PropertyAnimatorCore/Source/PropertyAnimatorCoreEditor/Private/Customizations/PropertyAnimatorCoreEditorContextTypeCustomization.h
@@ -26,16 +26,9 @@ public:
 	//~ End IPropertyTypeCustomization
 
 protected:
-	static UPropertyAnimatorCoreContext* GetPropertyContextValue(TSharedPtr<IPropertyHandle> InHandle);
-
 	ECheckBoxState IsPropertyEnabled() const;
 	void OnPropertyEnabled(ECheckBoxState InNewState) const;
-
-	/** Hide converter custom value widget */
-	void HideCustomValueWidget(IDetailPropertyRow* InConverterRow);
-
 	FReply UnlinkProperty() const;
 
 	TSharedPtr<IPropertyHandle> PropertyContextHandle;
-	TSharedPtr<IPropertyUtilities> PropertyUtilities;
 };
diff --git a/Engine/Plugins/Experimental/PropertyAnimatorCore/Source/PropertyAnimatorCoreEditor/Private/Customizations/PropertyAnimatorCoreEditorDetailCustomization.cpp b/Engine/Plugins/Experimental/PropertyAnimatorCore/Source/PropertyAnimatorCoreEditor/Private/Customizations/PropertyAnimatorCoreEditorDetailCustomization.cpp
--- a/Engine/Plugins/Experimental/PropertyAnimatorCore/Source/PropertyAnimatorCoreEditor/Private/Customizations/PropertyAnimatorCoreEditorDetailCustomization.cpp
+++ b/Engine/Plugins/Experimental/PropertyAnimatorCore/Source/PropertyAnimatorCoreEditor/Private/Customizations/PropertyAnimatorCoreEditorDetailCustomization.cpp
@@ -121,6 +121,14 @@ void FPropertyAnimatorCoreEditorDetailCustomization::CustomizeDetails(IDetailLay
 			[
 				LinkedPropertiesHandle->CreatePropertyNameWidget()
 			]
+			+ SHorizontalBox::Slot()
+			.VAlign(VAlign_Center)
+			.AutoWidth()
+			[
+				SNew(STextBlock)
+				.Font(InDetailBuilder.GetDetailFont())
+				.Text(this, &FPropertyAnimatorCoreEditorDetailCustomization::GetPropertiesCountText)
+			]
 		]
 		.ValueContent()
 		[
@@ -129,6 +137,45 @@ void FPropertyAnimatorCoreEditorDetailCustomization::CustomizeDetails(IDetailLay
 	}
 }
 
+int32 FPropertyAnimatorCoreEditorDetailCustomization::GetPropertiesCount() const
+{
+	int32 Count = INDEX_NONE;
+	for (const TWeakObjectPtr<UPropertyAnimatorCoreBase>& AnimatorWeak : AnimatorsWeak)
+	{
+		const UPropertyAnimatorCoreBase* Animator = AnimatorWeak.Get();
+		if (!Animator)
+		{
+			continue;
+		}
+
+		const int32 LinkedPropertiesCount = Animator->GetLinkedPropertiesCount();
+
+		if (Count == INDEX_NONE)
+		{
+			Count = LinkedPropertiesCount;
+		}
+		else if (Count != LinkedPropertiesCount)
+		{
+			Count = INDEX_NONE;
+			break;
+		}
+	}
+
+	return Count;
+}
+
+FText FPropertyAnimatorCoreEditorDetailCustomization::GetPropertiesCountText() const
+{
+	const int32 Count = GetPropertiesCount();
+
+	if (Count != INDEX_NONE)
+	{
+		return FText::Format(LOCTEXT("LinkedPropertiesCount", " ({0}) "), FText::FromString(FString::FromInt(Count)));
+	}
+
+	return LOCTEXT("LinkedPropertiesCountMismatch", " (multiple) ");
+}
+
 TSharedRef<SWidget> FPropertyAnimatorCoreEditorDetailCustomization::GenerateLinkMenu()
 {
 	UToolMenus* ToolMenus = UToolMenus::Get();
diff --git a/Engine/Plugins/Experimental/PropertyAnimatorCore/Source/PropertyAnimatorCoreEditor/Private/Customizations/PropertyAnimatorCoreEditorDetailCustomization.h b/Engine/Plugins/Experimental/PropertyAnimatorCore/Source/PropertyAnimatorCoreEditor/Private/Customizations/PropertyAnimatorCoreEditorDetailCustomization.h
--- a/Engine/Plugins/Experimental/PropertyAnimatorCore/Source/PropertyAnimatorCoreEditor/Private/Customizations/PropertyAnimatorCoreEditorDetailCustomization.h
+++ b/Engine/Plugins/Experimental/PropertyAnimatorCore/Source/PropertyAnimatorCoreEditor/Private/Customizations/PropertyAnimatorCoreEditorDetailCustomization.h
@@ -34,6 +34,9 @@ protected:
 	ECheckBoxState IsPropertiesEnabled() const;
 	void OnPropertiesEnabled(ECheckBoxState InNewState) const;
 
+	int32 GetPropertiesCount() const;
+	FText GetPropertiesCountText() const;
+
 	FReply UnlinkProperties() const;
 
 	FReply OnCreatePropertyPresetClicked() const;
