MotionDesign : Animators
- Fix multiple transaction created for the same operation (menu, operator stack, details view)

diff --git a/Engine/Plugins/Experimental/PropertyAnimatorCore/Source/PropertyAnimatorCore/Private/Subsystems/PropertyAnimatorCoreSubsystem.cpp b/Engine/Plugins/Experimental/PropertyAnimatorCore/Source/PropertyAnimatorCore/Private/Subsystems/PropertyAnimatorCoreSubsystem.cpp
--- a/Engine/Plugins/Experimental/PropertyAnimatorCore/Source/PropertyAnimatorCore/Private/Subsystems/PropertyAnimatorCoreSubsystem.cpp
+++ b/Engine/Plugins/Experimental/PropertyAnimatorCore/Source/PropertyAnimatorCore/Private/Subsystems/PropertyAnimatorCoreSubsystem.cpp
@@ -356,7 +356,7 @@ TSet<UPropertyAnimatorCoreBase*> UPropertyAnimatorCoreSubsystem::CreateAnimators
 	const FText AnimatorName = FText::FromName(AnimatorCDO->GetAnimatorOriginalName());
 	const FText ActorCount = FText::FromString(FString::FromInt(InActors.Num()));
 
-	FScopedTransaction Transaction(FText::Format(TransactionText, AnimatorName, ActorCount), bInTransact);
+	FScopedTransaction Transaction(FText::Format(TransactionText, AnimatorName, ActorCount), bInTransact && !GIsTransacting);
 
 	if (FEngineAnalytics::IsAvailable())
 	{
@@ -422,7 +422,7 @@ TSet<UPropertyAnimatorCoreBase*> UPropertyAnimatorCoreSubsystem::CloneAnimators(
 	const FText AnimatorCount = FText::FromString(FString::FromInt(InAnimators.Num()));
 	const FText ActorName = FText::FromString(InTargetActor->GetActorNameOrLabel());
 
-	FScopedTransaction Transaction(FText::Format(TransactionText, AnimatorCount, ActorName), bInTransact);
+	FScopedTransaction Transaction(FText::Format(TransactionText, AnimatorCount, ActorName), bInTransact && !GIsTransacting);
 #endif
 
 	UPropertyAnimatorCoreComponent* Component = UPropertyAnimatorCoreComponent::FindOrAdd(InTargetActor);
@@ -489,7 +489,7 @@ bool UPropertyAnimatorCoreSubsystem::RemoveAnimators(const TSet<UPropertyAnimato
 		const FText TransactionText = LOCTEXT("RemoveAnimator", "Removing {0} animator(s)");
 		const FText AnimatorCount = FText::FromString(FString::FromInt(InAnimators.Num()));
 
-		FScopedTransaction Transaction(FText::Format(TransactionText, AnimatorCount), bInTransact);
+		FScopedTransaction Transaction(FText::Format(TransactionText, AnimatorCount), bInTransact && !GIsTransacting);
 
 		Component->Modify();
 		Animator->Modify();
@@ -512,7 +512,7 @@ bool UPropertyAnimatorCoreSubsystem::RemoveAnimatorComponents(const TSet<UProper
 	const FText TransactionText = LOCTEXT("RemoveAnimatorComponent", "Removing {0} animator component(s)");
 	const FText ComponentCount = FText::FromString(FString::FromInt(InComponents.Num()));
 
-	FScopedTransaction Transaction(FText::Format(TransactionText, ComponentCount), bInTransact);
+	FScopedTransaction Transaction(FText::Format(TransactionText, ComponentCount), bInTransact && !GIsTransacting);
 #endif
 
 	for (UPropertyAnimatorCoreComponent* Component : InComponents)
@@ -554,7 +554,7 @@ bool UPropertyAnimatorCoreSubsystem::ApplyAnimatorPreset(UPropertyAnimatorCoreBa
 		const FText PresetName = FText::FromString(InPreset->GetPresetDisplayName());
 		const FText AnimatorName = FText::FromName(InAnimator->GetAnimatorOriginalName());
 
-		FScopedTransaction Transaction(FText::Format(TransactionText, PresetName, AnimatorName), bInTransact);
+		FScopedTransaction Transaction(FText::Format(TransactionText, PresetName, AnimatorName), bInTransact && !GIsTransacting);
 
 		InAnimator->Modify();
 #endif
@@ -579,7 +579,7 @@ bool UPropertyAnimatorCoreSubsystem::UnapplyAnimatorPreset(UPropertyAnimatorCore
 		const FText PresetName = FText::FromString(InPreset->GetPresetDisplayName());
 		const FText AnimatorName = FText::FromName(InAnimator->GetAnimatorOriginalName());
 
-		FScopedTransaction Transaction(FText::Format(TransactionText, PresetName, AnimatorName), bInTransact);
+		FScopedTransaction Transaction(FText::Format(TransactionText, PresetName, AnimatorName), bInTransact && !GIsTransacting);
 
 		InAnimator->Modify();
 #endif
@@ -631,7 +631,7 @@ bool UPropertyAnimatorCoreSubsystem::LinkAnimatorProperties(UPropertyAnimatorCor
 			);
 		}
 
-		FScopedTransaction Transaction(TransactionText, bInTransact);
+		FScopedTransaction Transaction(TransactionText, bInTransact && !GIsTransacting);
 
 		InAnimator->Modify();
 #endif
@@ -689,7 +689,7 @@ bool UPropertyAnimatorCoreSubsystem::UnlinkAnimatorProperties(UPropertyAnimatorC
 			);
 		}
 
-		FScopedTransaction Transaction(TransactionText, bInTransact);
+		FScopedTransaction Transaction(TransactionText, bInTransact && !GIsTransacting);
 
 		InAnimator->Modify();
 #endif
@@ -706,6 +706,86 @@ bool UPropertyAnimatorCoreSubsystem::UnlinkAnimatorProperties(UPropertyAnimatorC
 	return false;
 }
 
+bool UPropertyAnimatorCoreSubsystem::UnlinkAnimatorProperties(const TSet<UPropertyAnimatorCoreContext*>& InPropertyContexts, bool bInTransact)
+{
+	bool bResult = false;
+	if (InPropertyContexts.IsEmpty())
+	{
+		return bResult;
+	}
+
+#if WITH_EDITOR
+	FScopedTransaction Transaction(
+		FText::Format(LOCTEXT("UnlinkingAnimatorPropertyContexts", "Unlinking {0} properties from their animators"),
+			FText::FromString(FString::FromInt(InPropertyContexts.Num())))
+		, bInTransact && !GIsTransacting);
+#endif
+
+	for (const UPropertyAnimatorCoreContext* PropertyContext : InPropertyContexts)
+	{
+		if (PropertyContext)
+		{
+			if (UPropertyAnimatorCoreBase* Animator = PropertyContext->GetAnimator())
+			{
+#if WITH_EDITOR
+				Animator->Modify();
+#endif
+				bResult |= Animator->UnlinkProperty(PropertyContext->GetAnimatedProperty());
+			}
+		}
+	}
+
+#if WITH_EDITOR
+	if (!bResult)
+	{
+		Transaction.Cancel();		
+	}
+#endif
+
+	return bResult;
+}
+
+void UPropertyAnimatorCoreSubsystem::SetAnimatorPropertiesEnabled(const TSet<UPropertyAnimatorCoreContext*>& InPropertyContexts, bool bInEnabled, bool bInTransact)
+{
+	if (InPropertyContexts.IsEmpty())
+	{
+		return;
+	}
+
+#if WITH_EDITOR
+	const FText TransactionText = bInEnabled
+		? LOCTEXT("SetAnimatorPropertiesEnabled", "{0} animator properties enabled")
+		: LOCTEXT("SetAnimatorPropertiesDisabled", "{0} animator properties disabled");
+	
+	FScopedTransaction Transaction(
+		FText::Format(TransactionText,
+			FText::FromString(FString::FromInt(InPropertyContexts.Num())))
+		, bInTransact && !GIsTransacting);
+#endif
+
+	bool bResult = false;
+	for (UPropertyAnimatorCoreContext* PropertyContext : InPropertyContexts)
+	{
+		if (PropertyContext && PropertyContext->IsAnimated() != bInEnabled)
+		{
+#if WITH_EDITOR
+			PropertyContext->Modify();
+#endif
+
+			PropertyContext->SetAnimated(bInEnabled);
+			
+			bResult |= true;
+		}
+	}
+
+#if WITH_EDITOR
+	if (!bResult)
+	{
+		Transaction.Cancel();		
+	}
+#endif
+}
+
 bool UPropertyAnimatorCoreSubsystem::RegisterHandlerClass(const UClass* InHandlerClass)
 {
 	if (!IsValid(InHandlerClass))
@@ -1280,7 +1360,7 @@ void UPropertyAnimatorCoreSubsystem::SetActorAnimatorsEnabled(const TSet<AActor*
 		? LOCTEXT("SetActorAnimatorsEnabled", "Actors Animators Enabled")
 		: LOCTEXT("SetActorAnimatorsDisabled", "Actors Animators Disabled");
 
-	FScopedTransaction Transaction(TransactionText, bInTransact);
+	FScopedTransaction Transaction(TransactionText, bInTransact && !GIsTransacting);
 #endif
 
 	for (const AActor* Actor : InActors)
@@ -1317,7 +1397,7 @@ void UPropertyAnimatorCoreSubsystem::SetLevelAnimatorsEnabled(const UWorld* InWo
 		? LOCTEXT("SetLevelAnimatorsEnabled", "Level Animators Enabled")
 		: LOCTEXT("SetLevelAnimatorsDisabled", "Level Animators Disabled");
 
-	FScopedTransaction Transaction(TransactionText, bInTransact);
+	FScopedTransaction Transaction(TransactionText, bInTransact && !GIsTransacting);
 #endif
 
 	OnAnimatorsSetEnabledDelegate.Broadcast(InWorld, bInEnabled, bInTransact);
@@ -1335,7 +1415,7 @@ void UPropertyAnimatorCoreSubsystem::SetAnimatorsEnabled(const TSet<UPropertyAni
 		? LOCTEXT("SetAnimatorsEnabled", "{0} Animators Enabled")
 		: LOCTEXT("SetAnimatorsDisabled", "{0} Animators Disabled");
 
-	FScopedTransaction Transaction(FText::Format(TransactionText, FText::FromString(FString::FromInt(InAnimators.Num()))), bInTransact);
+	FScopedTransaction Transaction(FText::Format(TransactionText, FText::FromString(FString::FromInt(InAnimators.Num()))), bInTransact && !GIsTransacting);
 #endif
 
 	for (UPropertyAnimatorCoreBase* Animator : InAnimators)
diff --git a/Engine/Plugins/Experimental/PropertyAnimatorCore/Source/PropertyAnimatorCore/Public/Subsystems/PropertyAnimatorCoreSubsystem.h b/Engine/Plugins/Experimental/PropertyAnimatorCore/Source/PropertyAnimatorCore/Public/Subsystems/PropertyAnimatorCoreSubsystem.h
--- a/Engine/Plugins/Experimental/PropertyAnimatorCore/Source/PropertyAnimatorCore/Public/Subsystems/PropertyAnimatorCoreSubsystem.h
+++ b/Engine/Plugins/Experimental/PropertyAnimatorCore/Source/PropertyAnimatorCore/Public/Subsystems/PropertyAnimatorCoreSubsystem.h
@@ -187,6 +187,10 @@ public:
 	/** Unlink a property from an existing animator */
 	PROPERTYANIMATORCORE_API bool UnlinkAnimatorProperty(UPropertyAnimatorCoreBase* InAnimator, FPropertyAnimatorCoreData& InProperty, bool bInTransact = false);
 	PROPERTYANIMATORCORE_API bool UnlinkAnimatorProperties(UPropertyAnimatorCoreBase* InAnimator, const TSet<FPropertyAnimatorCoreData>& InProperties, bool bInTransact = false);
+	PROPERTYANIMATORCORE_API bool UnlinkAnimatorProperties(const TSet<UPropertyAnimatorCoreContext*>& InPropertyContexts, bool bInTransact = false);
+
+	/** Set the enabled state of animator property context */
+	PROPERTYANIMATORCORE_API void SetAnimatorPropertiesEnabled(const TSet<UPropertyAnimatorCoreContext*>& InPropertyContexts, bool bInEnabled, bool bInTransact = false);
 
 	/** Set the enabled state of animators attached to actors, will disable state globally on the component */
 	PROPERTYANIMATORCORE_API void SetActorAnimatorsEnabled(const TSet<AActor*>& InActors, bool bInEnabled, bool bInTransact = false);
diff --git a/Engine/Plugins/Experimental/PropertyAnimatorCore/Source/PropertyAnimatorCoreEditor/Private/Customizations/PropertyAnimatorCoreEditorContextTypeCustomization.cpp b/Engine/Plugins/Experimental/PropertyAnimatorCore/Source/PropertyAnimatorCoreEditor/Private/Customizations/PropertyAnimatorCoreEditorContextTypeCustomization.cpp
--- a/Engine/Plugins/Experimental/PropertyAnimatorCore/Source/PropertyAnimatorCoreEditor/Private/Customizations/PropertyAnimatorCoreEditorContextTypeCustomization.cpp
+++ b/Engine/Plugins/Experimental/PropertyAnimatorCore/Source/PropertyAnimatorCoreEditor/Private/Customizations/PropertyAnimatorCoreEditorContextTypeCustomization.cpp
@@ -154,12 +154,16 @@ void FPropertyAnimatorCoreEditorContextTypeCustomization::OnPropertyEnabled(EChe
 	TArray<UObject*> Objects;
 	PropertyContextHandle->GetOuterObjects(Objects);
 
-	for (UObject* Object : Objects)
+	TSet<UPropertyAnimatorCoreContext*> PropertyContexts;
+	Algo::Transform(Objects, PropertyContexts, [](UObject* InObject)
 	{
-		if (UPropertyAnimatorCoreContext* ContextObject = Cast<UPropertyAnimatorCoreContext>(Object))
-		{
-			ContextObject->SetAnimated(InNewState == ECheckBoxState::Checked);
-		}
+		return Cast<UPropertyAnimatorCoreContext>(InObject);
+	});
+
+	if (UPropertyAnimatorCoreSubsystem* AnimatorSubsystem = UPropertyAnimatorCoreSubsystem::Get())
+	{
+		constexpr bool bShouldTransact = true;
+		AnimatorSubsystem->SetAnimatorPropertiesEnabled(PropertyContexts, InNewState == ECheckBoxState::Checked, bShouldTransact);
 	}
 }
 
@@ -170,21 +174,19 @@ FReply FPropertyAnimatorCoreEditorContextTypeCustomization::UnlinkProperty() con
 		return FReply::Unhandled();
 	}
 
-	UPropertyAnimatorCoreSubsystem* AnimatorSubsystem = UPropertyAnimatorCoreSubsystem::Get();
 	TArray<UObject*> Objects;
 	PropertyContextHandle->GetOuterObjects(Objects);
 
-	constexpr bool bShouldTransact = true;
-	for (UObject* Object : Objects)
+	TSet<UPropertyAnimatorCoreContext*> PropertyContexts;
+	Algo::Transform(Objects, PropertyContexts, [](UObject* InObject)
 	{
-		if (const UPropertyAnimatorCoreContext* ContextObject = Cast<UPropertyAnimatorCoreContext>(Object))
-		{
-			if (UPropertyAnimatorCoreBase* PropertyAnimator = ContextObject->GetTypedOuter<UPropertyAnimatorCoreBase>())
-			{
-				FPropertyAnimatorCoreData Property = ContextObject->GetAnimatedProperty();
-				AnimatorSubsystem->UnlinkAnimatorProperty(PropertyAnimator, Property, bShouldTransact);
-			}
-		}
+		return Cast<UPropertyAnimatorCoreContext>(InObject);
+	});
+
+	if (UPropertyAnimatorCoreSubsystem* AnimatorSubsystem = UPropertyAnimatorCoreSubsystem::Get())
+	{
+		constexpr bool bShouldTransact = true;
+		AnimatorSubsystem->UnlinkAnimatorProperties(PropertyContexts, bShouldTransact);
 	}
 
 	return FReply::Handled();
diff --git a/Engine/Plugins/Experimental/PropertyAnimatorCore/Source/PropertyAnimatorCoreEditor/Private/Customizations/PropertyAnimatorCoreEditorDetailCustomization.cpp b/Engine/Plugins/Experimental/PropertyAnimatorCore/Source/PropertyAnimatorCoreEditor/Private/Customizations/PropertyAnimatorCoreEditorDetailCustomization.cpp
--- a/Engine/Plugins/Experimental/PropertyAnimatorCore/Source/PropertyAnimatorCoreEditor/Private/Customizations/PropertyAnimatorCoreEditorDetailCustomization.cpp
+++ b/Engine/Plugins/Experimental/PropertyAnimatorCore/Source/PropertyAnimatorCoreEditor/Private/Customizations/PropertyAnimatorCoreEditorDetailCustomization.cpp
@@ -322,6 +322,7 @@ void FPropertyAnimatorCoreEditorDetailCustomization::OnPropertiesEnabled(ECheckB
 		return;
 	}
 
+	TSet<UPropertyAnimatorCoreContext*> PropertyContexts;
 	for (const TWeakObjectPtr<UPropertyAnimatorCoreBase>& AnimatorWeak : AnimatorsWeak)
 	{
 		const UPropertyAnimatorCoreBase* Animator = AnimatorWeak.Get();
@@ -331,29 +332,38 @@ void FPropertyAnimatorCoreEditorDetailCustomization::OnPropertiesEnabled(ECheckB
 		}
 
 		constexpr bool bResolve = false;
-		Animator->ForEachLinkedProperty<UPropertyAnimatorCoreContext>([InNewState](UPropertyAnimatorCoreContext* InOptions, const FPropertyAnimatorCoreData& InProperty)->bool
+		Animator->ForEachLinkedProperty<UPropertyAnimatorCoreContext>([&PropertyContexts](UPropertyAnimatorCoreContext* InOptions, const FPropertyAnimatorCoreData& InProperty)->bool
 		{
-			InOptions->SetAnimated(InNewState == ECheckBoxState::Checked);
+			PropertyContexts.Add(InOptions);
 			return true;
 		}, bResolve);
 	}
+
+	if (UPropertyAnimatorCoreSubsystem* AnimatorSubsystem = UPropertyAnimatorCoreSubsystem::Get())
+	{
+		constexpr bool bShouldTransact = true;
+		AnimatorSubsystem->SetAnimatorPropertiesEnabled(PropertyContexts, InNewState == ECheckBoxState::Checked, bShouldTransact);
+	}
 }
 
 FReply FPropertyAnimatorCoreEditorDetailCustomization::UnlinkProperties() const
 {
 	if (UPropertyAnimatorCoreSubsystem* AnimatorSubsystem = UPropertyAnimatorCoreSubsystem::Get())
 	{
+		TSet<UPropertyAnimatorCoreContext*> PropertyContexts;
 		for (const TWeakObjectPtr<UPropertyAnimatorCoreBase>& AnimatorWeak : AnimatorsWeak)
 		{
-			UPropertyAnimatorCoreBase* Animator = AnimatorWeak.Get();	
+			const UPropertyAnimatorCoreBase* Animator = AnimatorWeak.Get();	
 			if (!Animator)
 			{
 				continue;
 			}
-		
-			const TSet<FPropertyAnimatorCoreData> LinkedProperties = Animator->GetLinkedProperties();
-			AnimatorSubsystem->UnlinkAnimatorProperties(Animator, LinkedProperties, true);
+
+			PropertyContexts.Append(Animator->GetLinkedPropertiesContext());
 		}
+
+		constexpr bool bShouldTransact = true;
+		AnimatorSubsystem->UnlinkAnimatorProperties(PropertyContexts, bShouldTransact);
 	}
 
 	return FReply::Handled();
diff --git a/Engine/Plugins/Experimental/PropertyAnimatorCore/Source/PropertyAnimatorCoreEditor/Private/Menus/PropertyAnimatorCoreEditorMenu.cpp b/Engine/Plugins/Experimental/PropertyAnimatorCore/Source/PropertyAnimatorCoreEditor/Private/Menus/PropertyAnimatorCoreEditorMenu.cpp
--- a/Engine/Plugins/Experimental/PropertyAnimatorCore/Source/PropertyAnimatorCoreEditor/Private/Menus/PropertyAnimatorCoreEditorMenu.cpp
+++ b/Engine/Plugins/Experimental/PropertyAnimatorCore/Source/PropertyAnimatorCoreEditor/Private/Menus/PropertyAnimatorCoreEditorMenu.cpp
@@ -10,6 +10,7 @@
 #include "Presets/PropertyAnimatorCoreAnimatorPreset.h"
 #include "Presets/PropertyAnimatorCorePresetBase.h"
 #include "Presets/PropertyAnimatorCorePropertyPreset.h"
+#include "ScopedTransaction.h"
 #include "Styling/SlateIconFinder.h"
 #include "Subsystems/PropertyAnimatorCoreEditorSubsystem.h"
 #include "Subsystems/PropertyAnimatorCoreSubsystem.h"
@@ -395,12 +396,20 @@ void UE::PropertyAnimatorCoreEditor::Menu::ExecuteNewAnimatorPropertyAction(cons
 		return;
 	}
 
+	const TSet<AActor*>& ContextActors = InMenuData->GetContext().GetActors();
+	const bool bShouldTransact = InMenuData->GetOptions().ShouldTransact();
+
+	FScopedTransaction Transaction(
+		FText::Format(LOCTEXT("CreateAnimatorWithProperty", "Create animator with property on {0} actors"),
+			FText::AsNumber(ContextActors.Num()))
+		, bShouldTransact && !GIsTransacting);
+
 	TSet<UPropertyAnimatorCoreBase*> CreatedAnimators;
-	for (AActor* Actor : InMenuData->GetContext().GetActors())
+	for (AActor* Actor : ContextActors)
 	{
-		UPropertyAnimatorCoreBase* NewAnimator = Subsystem->CreateAnimator(Actor, InAnimator->GetClass(), nullptr, InMenuData->GetOptions().ShouldTransact());
+		UPropertyAnimatorCoreBase* NewAnimator = Subsystem->CreateAnimator(Actor, InAnimator->GetClass(), nullptr, bShouldTransact);
 		FPropertyAnimatorCoreData Property(Actor, InPropertyLocatorPath);
-		Subsystem->LinkAnimatorProperty(NewAnimator, Property, true);
+		Subsystem->LinkAnimatorProperty(NewAnimator, Property, bShouldTransact);
 		CreatedAnimators.Add(NewAnimator);
 	}
 	
@@ -418,8 +427,16 @@ void UE::PropertyAnimatorCoreEditor::Menu::ExecuteLinkLastCreatedAnimatorPropert
 		return;
 	}
 
-	constexpr bool bShouldTransact = true;
-	for (UPropertyAnimatorCoreBase* LastCreatedAnimator : InMenuData->GetLastCreatedAnimators())
+	const TSet<UPropertyAnimatorCoreBase*> LastAnimators = InMenuData->GetLastCreatedAnimators();
+
+	const bool bShouldTransact = InMenuData->GetOptions().ShouldTransact();
+
+	FScopedTransaction Transaction(
+		FText::Format(LOCTEXT("LinkPropertyToAnimators", "Toggle link property to {0} animator(s)"),
+			FText::AsNumber(LastAnimators.Num()))
+		, bShouldTransact && !GIsTransacting);
+
+	for (UPropertyAnimatorCoreBase* LastCreatedAnimator : LastAnimators)
 	{
 		if (LastCreatedAnimator->GetClass() != InAnimator->GetClass())
 		{
@@ -449,7 +466,17 @@ void UE::PropertyAnimatorCoreEditor::Menu::ExecuteApplyLastCreatedAnimatorPreset
 		return;
 	}
 
-	for (UPropertyAnimatorCoreBase* LastCreatedAnimator : InMenuData->GetLastCreatedAnimators())
+	const TSet<UPropertyAnimatorCoreBase*> LastAnimators = InMenuData->GetLastCreatedAnimators();
+
+	const bool bShouldTransact = InMenuData->GetOptions().ShouldTransact();
+
+	FScopedTransaction Transaction(
+		FText::Format(LOCTEXT("ApplyAnimatorsPreset", "Toggle apply preset {0} on {1} animator(s)"),
+			FText::FromString(InPreset->GetPresetDisplayName()),
+			FText::AsNumber(LastAnimators.Num()))
+		, bShouldTransact && !GIsTransacting);
+
+	for (UPropertyAnimatorCoreBase* LastCreatedAnimator : LastAnimators)
 	{
 		if (LastCreatedAnimator->GetClass() != InAnimator->GetClass())
 		{
@@ -458,11 +485,11 @@ void UE::PropertyAnimatorCoreEditor::Menu::ExecuteApplyLastCreatedAnimatorPreset
 
 		if (InPreset->IsPresetApplied(LastCreatedAnimator))
 		{
-			Subsystem->UnapplyAnimatorPreset(LastCreatedAnimator, InPreset, true);
+			Subsystem->UnapplyAnimatorPreset(LastCreatedAnimator, InPreset, bShouldTransact);
 		}
 		else
 		{
-			Subsystem->ApplyAnimatorPreset(LastCreatedAnimator, InPreset, true);
+			Subsystem->ApplyAnimatorPreset(LastCreatedAnimator, InPreset, bShouldTransact);
 		}
 	}
 }
@@ -481,6 +508,14 @@ void UE::PropertyAnimatorCoreEditor::Menu::ExecuteLinkAnimatorPresetAction(const
 
 	const ECheckBoxState State = GetAnimatorPresetState(InAnimators, InPreset);
 
+	const bool bShouldTransact = InMenuData->GetOptions().ShouldTransact();
+
+	FScopedTransaction Transaction(
+		FText::Format(LOCTEXT("ApplyAnimatorsPreset", "Toggle apply preset {0} on {1} animator(s)"),
+			FText::FromString(InPreset->GetPresetDisplayName()),
+			FText::AsNumber(InAnimators.Num()))
+		, bShouldTransact && !GIsTransacting);
+
 	for (UPropertyAnimatorCoreBase* Animator : InAnimators)
 	{
 		if (!IsValid(Animator) || Animator->IsTemplate())
@@ -490,11 +525,11 @@ void UE::PropertyAnimatorCoreEditor::Menu::ExecuteLinkAnimatorPresetAction(const
 
 		if (State == ECheckBoxState::Undetermined || InPreset->IsPresetApplied(Animator))
 		{
-			Subsystem->UnapplyAnimatorPreset(Animator, InPreset, InMenuData->GetOptions().ShouldTransact());
+			Subsystem->UnapplyAnimatorPreset(Animator, InPreset, bShouldTransact);
 		}
 		else
 		{
-			Subsystem->ApplyAnimatorPreset(Animator, InPreset, InMenuData->GetOptions().ShouldTransact());
+			Subsystem->ApplyAnimatorPreset(Animator, InPreset, bShouldTransact);
 		}
 	}
 }
@@ -512,6 +547,13 @@ void UE::PropertyAnimatorCoreEditor::Menu::ExecuteLinkAnimatorPropertyAction(con
 
 	const ECheckBoxState State = GetAnimatorPropertyLinkState(InAnimators, InPropertyLocatorPath);
 
+	const bool bShouldTransact = InMenuData->GetOptions().ShouldTransact();
+
+	FScopedTransaction Transaction(
+		FText::Format(LOCTEXT("LinkPropertyToAnimators", "Toggle link property to {0} animator(s)"),
+			FText::AsNumber(InAnimators.Num()))
+		, bShouldTransact && !GIsTransacting);
+	
 	for (UPropertyAnimatorCoreBase* Animator : InAnimators)
 	{
 		if (!IsValid(Animator) || Animator->IsTemplate())
@@ -528,7 +570,7 @@ void UE::PropertyAnimatorCoreEditor::Menu::ExecuteLinkAnimatorPropertyAction(con
 
 		if (State == ECheckBoxState::Undetermined || Animator->IsPropertyLinked(AnimatorProperty))
 		{
-			Subsystem->UnlinkAnimatorProperty(Animator, AnimatorProperty, InMenuData->GetOptions().ShouldTransact());
+			Subsystem->UnlinkAnimatorProperty(Animator, AnimatorProperty, bShouldTransact);
 
 			if (InPreset)
 			{
@@ -537,7 +579,7 @@ void UE::PropertyAnimatorCoreEditor::Menu::ExecuteLinkAnimatorPropertyAction(con
 		}
 		else
 		{
-			Subsystem->LinkAnimatorProperty(Animator, AnimatorProperty, InMenuData->GetOptions().ShouldTransact());
+			Subsystem->LinkAnimatorProperty(Animator, AnimatorProperty, bShouldTransact);
 
 			if (InPreset)
 			{
