MotionDesign : PropertyAnimator
- Fix crash caused by property bag not adding values due to invalid names used, locator path hash will be used instead
- Migrating delta and original values from property bags to new naming scheme

diff --git a/Engine/Plugins/Experimental/PropertyAnimator/Source/PropertyAnimator/Private/Animators/PropertyAnimatorClock.cpp b/Engine/Plugins/Experimental/PropertyAnimator/Source/PropertyAnimator/Private/Animators/PropertyAnimatorClock.cpp
--- a/Engine/Plugins/Experimental/PropertyAnimator/Source/PropertyAnimator/Private/Animators/PropertyAnimatorClock.cpp
+++ b/Engine/Plugins/Experimental/PropertyAnimator/Source/PropertyAnimator/Private/Animators/PropertyAnimatorClock.cpp
@@ -107,10 +107,10 @@ void UPropertyAnimatorClock::EvaluateProperties(FInstancedPropertyBag& InParamet
 		, int32 InRangeIndex
 		, int32 InRangeMax)->bool
 	{
-		const FName DisplayName(InResolvedProperty.GetPathHash());
+		const FName PropertyHash = InResolvedProperty.GetLocatorPathHash();
 
-		InEvaluatedValues.AddProperty(DisplayName, EPropertyBagPropertyType::String);
-		InEvaluatedValues.SetValueString(DisplayName, FormattedDateTime);
+		InEvaluatedValues.AddProperty(PropertyHash, EPropertyBagPropertyType::String);
+		InEvaluatedValues.SetValueString(PropertyHash, FormattedDateTime);
 
 		return true;
 	});
diff --git a/Engine/Plugins/Experimental/PropertyAnimator/Source/PropertyAnimator/Private/Animators/PropertyAnimatorCounter.cpp b/Engine/Plugins/Experimental/PropertyAnimator/Source/PropertyAnimator/Private/Animators/PropertyAnimatorCounter.cpp
--- a/Engine/Plugins/Experimental/PropertyAnimator/Source/PropertyAnimator/Private/Animators/PropertyAnimatorCounter.cpp
+++ b/Engine/Plugins/Experimental/PropertyAnimator/Source/PropertyAnimator/Private/Animators/PropertyAnimatorCounter.cpp
@@ -272,10 +272,10 @@ void UPropertyAnimatorCounter::EvaluateProperties(FInstancedPropertyBag& InParam
 		, int32 InRangeIndex
 		, int32 InRangeMax)->bool
 	{
-		const FName DisplayName(InResolvedProperty.GetPathHash());
+		const FName PropertyHash = InResolvedProperty.GetLocatorPathHash();
 
-		InEvaluatedValues.AddProperty(DisplayName, EPropertyBagPropertyType::String);
-		InEvaluatedValues.SetValueString(DisplayName, FormatNumber(TimeElapsed));
+		InEvaluatedValues.AddProperty(PropertyHash, EPropertyBagPropertyType::String);
+		InEvaluatedValues.SetValueString(PropertyHash, FormatNumber(TimeElapsed));
 
 		return true;
 	});
diff --git a/Engine/Plugins/Experimental/PropertyAnimator/Source/PropertyAnimator/Private/Properties/PropertyAnimatorFloatContext.cpp b/Engine/Plugins/Experimental/PropertyAnimator/Source/PropertyAnimator/Private/Properties/PropertyAnimatorFloatContext.cpp
--- a/Engine/Plugins/Experimental/PropertyAnimator/Source/PropertyAnimator/Private/Properties/PropertyAnimatorFloatContext.cpp
+++ b/Engine/Plugins/Experimental/PropertyAnimator/Source/PropertyAnimator/Private/Properties/PropertyAnimatorFloatContext.cpp
@@ -39,9 +39,9 @@ bool UPropertyAnimatorFloatContext::EvaluateProperty(const FPropertyAnimatorCore
 
 	if (AlphaResult.HasValue())
 	{
-		const FName DisplayName(InProperty.GetPathHash());
-		OutEvaluatedValues.AddProperty(DisplayName, EPropertyBagPropertyType::Double);
-		OutEvaluatedValues.SetValueFloat(DisplayName, MagnitudeResult.GetValue() * FMath::Lerp(AmplitudeMin, AmplitudeMax, AlphaResult.GetValue()));
+		const FName PropertyHash = InProperty.GetLocatorPathHash();
+		OutEvaluatedValues.AddProperty(PropertyHash, EPropertyBagPropertyType::Double);
+		OutEvaluatedValues.SetValueFloat(PropertyHash, MagnitudeResult.GetValue() * FMath::Lerp(AmplitudeMin, AmplitudeMax, AlphaResult.GetValue()));
 		return true;
 	}
 
diff --git a/Engine/Plugins/Experimental/PropertyAnimator/Source/PropertyAnimator/Private/Properties/PropertyAnimatorRotatorContext.cpp b/Engine/Plugins/Experimental/PropertyAnimator/Source/PropertyAnimator/Private/Properties/PropertyAnimatorRotatorContext.cpp
--- a/Engine/Plugins/Experimental/PropertyAnimator/Source/PropertyAnimator/Private/Properties/PropertyAnimatorRotatorContext.cpp
+++ b/Engine/Plugins/Experimental/PropertyAnimator/Source/PropertyAnimator/Private/Properties/PropertyAnimatorRotatorContext.cpp
@@ -39,9 +39,9 @@ bool UPropertyAnimatorRotatorContext::EvaluateProperty(const FPropertyAnimatorCo
 
 	if (AlphaResult.HasValue())
 	{
-		const FName DisplayName(InProperty.GetPathHash());
-		OutEvaluatedValues.AddProperty(DisplayName, EPropertyBagPropertyType::Struct, TBaseStructure<FRotator>::Get());
-		OutEvaluatedValues.SetValueStruct(DisplayName, MagnitudeResult.GetValue() * FMath::LerpRange(AmplitudeMin, AmplitudeMax, AlphaResult.GetValue()));
+		const FName PropertyHash = InProperty.GetLocatorPathHash();
+		OutEvaluatedValues.AddProperty(PropertyHash, EPropertyBagPropertyType::Struct, TBaseStructure<FRotator>::Get());
+		OutEvaluatedValues.SetValueStruct(PropertyHash, MagnitudeResult.GetValue() * FMath::LerpRange(AmplitudeMin, AmplitudeMax, AlphaResult.GetValue()));
 		return true;
 	}
 
diff --git a/Engine/Plugins/Experimental/PropertyAnimator/Source/PropertyAnimator/Private/Properties/PropertyAnimatorVectorContext.cpp b/Engine/Plugins/Experimental/PropertyAnimator/Source/PropertyAnimator/Private/Properties/PropertyAnimatorVectorContext.cpp
--- a/Engine/Plugins/Experimental/PropertyAnimator/Source/PropertyAnimator/Private/Properties/PropertyAnimatorVectorContext.cpp
+++ b/Engine/Plugins/Experimental/PropertyAnimator/Source/PropertyAnimator/Private/Properties/PropertyAnimatorVectorContext.cpp
@@ -39,9 +39,9 @@ bool UPropertyAnimatorVectorContext::EvaluateProperty(const FPropertyAnimatorCor
 
 	if (AlphaResult.HasValue())
 	{
-		const FName DisplayName(InProperty.GetPathHash());
-		OutEvaluatedValues.AddProperty(DisplayName, EPropertyBagPropertyType::Struct, TBaseStructure<FVector>::Get());
-		OutEvaluatedValues.SetValueStruct(DisplayName, MagnitudeResult.GetValue() * FMath::Lerp(AmplitudeMin, AmplitudeMax, AlphaResult.GetValue()));
+		const FName PropertyHash = InProperty.GetLocatorPathHash();
+		OutEvaluatedValues.AddProperty(PropertyHash, EPropertyBagPropertyType::Struct, TBaseStructure<FVector>::Get());
+		OutEvaluatedValues.SetValueStruct(PropertyHash, MagnitudeResult.GetValue() * FMath::Lerp(AmplitudeMin, AmplitudeMax, AlphaResult.GetValue()));
 		return true;
 	}
 
diff --git a/Engine/Plugins/Experimental/PropertyAnimatorCore/Source/PropertyAnimatorCore/Private/Properties/Handlers/PropertyAnimatorCoreBoolHandler.cpp b/Engine/Plugins/Experimental/PropertyAnimatorCore/Source/PropertyAnimatorCore/Private/Properties/Handlers/PropertyAnimatorCoreBoolHandler.cpp
--- a/Engine/Plugins/Experimental/PropertyAnimatorCore/Source/PropertyAnimatorCore/Private/Properties/Handlers/PropertyAnimatorCoreBoolHandler.cpp
+++ b/Engine/Plugins/Experimental/PropertyAnimatorCore/Source/PropertyAnimatorCore/Private/Properties/Handlers/PropertyAnimatorCoreBoolHandler.cpp
@@ -14,27 +14,27 @@ bool UPropertyAnimatorCoreBoolHandler::IsPropertySupported(const FPropertyAnimat
 
 bool UPropertyAnimatorCoreBoolHandler::GetValue(const FPropertyAnimatorCoreData& InPropertyData, FInstancedPropertyBag& OutValue)
 {
-	const FName PropertyName(InPropertyData.GetPathHash());
-	OutValue.AddProperty(PropertyName, EPropertyBagPropertyType::Bool);
+	const FName PropertyHash(InPropertyData.GetLocatorPathHash());
+	OutValue.AddProperty(PropertyHash, EPropertyBagPropertyType::Bool);
 
 	bool bValue;
 	InPropertyData.GetPropertyValuePtr(&bValue);
 
-	OutValue.SetValueBool(PropertyName, bValue);
+	OutValue.SetValueBool(PropertyHash, bValue);
 
 	return true;
 }
 
 bool UPropertyAnimatorCoreBoolHandler::SetValue(const FPropertyAnimatorCoreData& InPropertyData, const FInstancedPropertyBag& InValue)
 {
-	const FName PropertyName(InPropertyData.GetPathHash());
-	TValueOrError<bool, EPropertyBagResult> ValueResult = InValue.GetValueBool(PropertyName);
+	const FName PropertyHash(InPropertyData.GetLocatorPathHash());
+	TValueOrError<bool, EPropertyBagResult> ValueResult = InValue.GetValueBool(PropertyHash);
 	if (!ValueResult.HasValue())
 	{
 		return false;
 	}
 
-	bool bNewValue = ValueResult.GetValue();
+	bool& bNewValue = ValueResult.GetValue();
 	InPropertyData.SetPropertyValuePtr(&bNewValue);
 
 	return true;
diff --git a/Engine/Plugins/Experimental/PropertyAnimatorCore/Source/PropertyAnimatorCore/Private/Properties/Handlers/PropertyAnimatorCoreDoubleHandler.cpp b/Engine/Plugins/Experimental/PropertyAnimatorCore/Source/PropertyAnimatorCore/Private/Properties/Handlers/PropertyAnimatorCoreDoubleHandler.cpp
--- a/Engine/Plugins/Experimental/PropertyAnimatorCore/Source/PropertyAnimatorCore/Private/Properties/Handlers/PropertyAnimatorCoreDoubleHandler.cpp
+++ b/Engine/Plugins/Experimental/PropertyAnimatorCore/Source/PropertyAnimatorCore/Private/Properties/Handlers/PropertyAnimatorCoreDoubleHandler.cpp
@@ -19,27 +19,27 @@ bool UPropertyAnimatorCoreDoubleHandler::IsPropertySupported(const FPropertyAnim
 
 bool UPropertyAnimatorCoreDoubleHandler::GetValue(const FPropertyAnimatorCoreData& InPropertyData, FInstancedPropertyBag& OutValue)
 {
-	const FName PropertyName(InPropertyData.GetPathHash());
-	OutValue.AddProperty(PropertyName, EPropertyBagPropertyType::Double);
+	const FName PropertyHash(InPropertyData.GetLocatorPathHash());
+	OutValue.AddProperty(PropertyHash, EPropertyBagPropertyType::Double);
 
 	double Value;
 	InPropertyData.GetPropertyValuePtr(&Value);
 
-	OutValue.SetValueDouble(PropertyName, Value);
+	OutValue.SetValueDouble(PropertyHash, Value);
 
 	return true;
 }
 
 bool UPropertyAnimatorCoreDoubleHandler::SetValue(const FPropertyAnimatorCoreData& InPropertyData, const FInstancedPropertyBag& InValue)
 {
-	const FName PropertyName(InPropertyData.GetPathHash());
-	TValueOrError<double, EPropertyBagResult> ValueResult = InValue.GetValueDouble(PropertyName);
+	const FName PropertyHash(InPropertyData.GetLocatorPathHash());
+	TValueOrError<double, EPropertyBagResult> ValueResult = InValue.GetValueDouble(PropertyHash);
 	if (!ValueResult.HasValue())
 	{
 		return false;
 	}
 
-	double NewValue = ValueResult.GetValue();
+	double& NewValue = ValueResult.GetValue();
 	InPropertyData.SetPropertyValuePtr(&NewValue);
 
 	return true;
@@ -52,8 +52,8 @@ bool UPropertyAnimatorCoreDoubleHandler::IsAdditiveSupported() const
 
 bool UPropertyAnimatorCoreDoubleHandler::AddValue(const FPropertyAnimatorCoreData& InPropertyData, const FInstancedPropertyBag& InValue)
 {
-	const FName PropertyName(InPropertyData.GetPathHash());
-	const TValueOrError<double, EPropertyBagResult> ValueResult = InValue.GetValueDouble(PropertyName);
+	const FName PropertyHash(InPropertyData.GetLocatorPathHash());
+	const TValueOrError<double, EPropertyBagResult> ValueResult = InValue.GetValueDouble(PropertyHash);
 	if (!ValueResult.HasValue())
 	{
 		return false;
@@ -70,8 +70,8 @@ bool UPropertyAnimatorCoreDoubleHandler::AddValue(const FPropertyAnimatorCoreDat
 
 bool UPropertyAnimatorCoreDoubleHandler::SubtractValue(const FPropertyAnimatorCoreData& InPropertyData, const FInstancedPropertyBag& InValue)
 {
-	const FName PropertyName(InPropertyData.GetPathHash());
-	const TValueOrError<double, EPropertyBagResult> ValueResult = InValue.GetValueDouble(PropertyName);
+	const FName PropertyHash(InPropertyData.GetLocatorPathHash());
+	const TValueOrError<double, EPropertyBagResult> ValueResult = InValue.GetValueDouble(PropertyHash);
 	if (!ValueResult.HasValue())
 	{
 		return false;
diff --git a/Engine/Plugins/Experimental/PropertyAnimatorCore/Source/PropertyAnimatorCore/Private/Properties/Handlers/PropertyAnimatorCoreFloatHandler.cpp b/Engine/Plugins/Experimental/PropertyAnimatorCore/Source/PropertyAnimatorCore/Private/Properties/Handlers/PropertyAnimatorCoreFloatHandler.cpp
--- a/Engine/Plugins/Experimental/PropertyAnimatorCore/Source/PropertyAnimatorCore/Private/Properties/Handlers/PropertyAnimatorCoreFloatHandler.cpp
+++ b/Engine/Plugins/Experimental/PropertyAnimatorCore/Source/PropertyAnimatorCore/Private/Properties/Handlers/PropertyAnimatorCoreFloatHandler.cpp
@@ -19,21 +19,21 @@ bool UPropertyAnimatorCoreFloatHandler::IsPropertySupported(const FPropertyAnima
 
 bool UPropertyAnimatorCoreFloatHandler::GetValue(const FPropertyAnimatorCoreData& InPropertyData, FInstancedPropertyBag& OutValue)
 {
-	const FName PropertyName(InPropertyData.GetPathHash());
-	OutValue.AddProperty(PropertyName, EPropertyBagPropertyType::Float);
+	const FName PropertyHash(InPropertyData.GetLocatorPathHash());
+	OutValue.AddProperty(PropertyHash, EPropertyBagPropertyType::Float);
 
 	float Value;
 	InPropertyData.GetPropertyValuePtr(&Value);
 
-	OutValue.SetValueFloat(PropertyName, Value);
+	OutValue.SetValueFloat(PropertyHash, Value);
 
 	return true;
 }
 
 bool UPropertyAnimatorCoreFloatHandler::SetValue(const FPropertyAnimatorCoreData& InPropertyData, const FInstancedPropertyBag& InValue)
 {
-	const FName PropertyName(InPropertyData.GetPathHash());
-	TValueOrError<float, EPropertyBagResult> ValueResult = InValue.GetValueFloat(PropertyName);
+	const FName PropertyHash(InPropertyData.GetLocatorPathHash());
+	TValueOrError<float, EPropertyBagResult> ValueResult = InValue.GetValueFloat(PropertyHash);
 	if (!ValueResult.HasValue())
 	{
 		return false;
@@ -52,8 +52,8 @@ bool UPropertyAnimatorCoreFloatHandler::IsAdditiveSupported() const
 
 bool UPropertyAnimatorCoreFloatHandler::AddValue(const FPropertyAnimatorCoreData& InPropertyData, const FInstancedPropertyBag& InValue)
 {
-	const FName PropertyName(InPropertyData.GetPathHash());
-	const TValueOrError<float, EPropertyBagResult> ValueResult = InValue.GetValueFloat(PropertyName);
+	const FName PropertyHash(InPropertyData.GetLocatorPathHash());
+	const TValueOrError<float, EPropertyBagResult> ValueResult = InValue.GetValueFloat(PropertyHash);
 	if (!ValueResult.HasValue())
 	{
 		return false;
@@ -70,8 +70,8 @@ bool UPropertyAnimatorCoreFloatHandler::AddValue(const FPropertyAnimatorCoreData
 
 bool UPropertyAnimatorCoreFloatHandler::SubtractValue(const FPropertyAnimatorCoreData& InPropertyData, const FInstancedPropertyBag& InValue)
 {
-	const FName PropertyName(InPropertyData.GetPathHash());
-	const TValueOrError<float, EPropertyBagResult> ValueResult = InValue.GetValueFloat(PropertyName);
+	const FName PropertyHash(InPropertyData.GetLocatorPathHash());
+	const TValueOrError<float, EPropertyBagResult> ValueResult = InValue.GetValueFloat(PropertyHash);
 	if (!ValueResult.HasValue())
 	{
 		return false;
diff --git a/Engine/Plugins/Experimental/PropertyAnimatorCore/Source/PropertyAnimatorCore/Private/Properties/Handlers/PropertyAnimatorCoreIntHandler.cpp b/Engine/Plugins/Experimental/PropertyAnimatorCore/Source/PropertyAnimatorCore/Private/Properties/Handlers/PropertyAnimatorCoreIntHandler.cpp
--- a/Engine/Plugins/Experimental/PropertyAnimatorCore/Source/PropertyAnimatorCore/Private/Properties/Handlers/PropertyAnimatorCoreIntHandler.cpp
+++ b/Engine/Plugins/Experimental/PropertyAnimatorCore/Source/PropertyAnimatorCore/Private/Properties/Handlers/PropertyAnimatorCoreIntHandler.cpp
@@ -14,27 +14,27 @@ bool UPropertyAnimatorCoreIntHandler::IsPropertySupported(const FPropertyAnimato
 
 bool UPropertyAnimatorCoreIntHandler::GetValue(const FPropertyAnimatorCoreData& InPropertyData, FInstancedPropertyBag& OutValue)
 {
-	const FName PropertyName(InPropertyData.GetPathHash());
-	OutValue.AddProperty(PropertyName, EPropertyBagPropertyType::Int32);
+	const FName PropertyHash(InPropertyData.GetLocatorPathHash());
+	OutValue.AddProperty(PropertyHash, EPropertyBagPropertyType::Int32);
 
 	int32 Value;
 	InPropertyData.GetPropertyValuePtr(&Value);
 
-	OutValue.SetValueInt32(PropertyName, Value);
+	OutValue.SetValueInt32(PropertyHash, Value);
 
 	return true;
 }
 
 bool UPropertyAnimatorCoreIntHandler::SetValue(const FPropertyAnimatorCoreData& InPropertyData, const FInstancedPropertyBag& InValue)
 {
-	const FName PropertyName(InPropertyData.GetPathHash());
-	TValueOrError<int32, EPropertyBagResult> ValueResult = InValue.GetValueInt32(PropertyName);
+	const FName PropertyHash(InPropertyData.GetLocatorPathHash());
+	TValueOrError<int32, EPropertyBagResult> ValueResult = InValue.GetValueInt32(PropertyHash);
 	if (!ValueResult.HasValue())
 	{
 		return false;
 	}
 
-	int32 NewValue = ValueResult.GetValue();
+	int32& NewValue = ValueResult.GetValue();
 	InPropertyData.SetPropertyValuePtr(&NewValue);
 
 	return true;
@@ -47,8 +47,8 @@ bool UPropertyAnimatorCoreIntHandler::IsAdditiveSupported() const
 
 bool UPropertyAnimatorCoreIntHandler::AddValue(const FPropertyAnimatorCoreData& InPropertyData, const FInstancedPropertyBag& InValue)
 {
-	const FName PropertyName(InPropertyData.GetPathHash());
-	const TValueOrError<int32, EPropertyBagResult> ValueResult = InValue.GetValueInt32(PropertyName);
+	const FName PropertyHash(InPropertyData.GetLocatorPathHash());
+	const TValueOrError<int32, EPropertyBagResult> ValueResult = InValue.GetValueInt32(PropertyHash);
 	if (!ValueResult.HasValue())
 	{
 		return false;
@@ -65,8 +65,8 @@ bool UPropertyAnimatorCoreIntHandler::AddValue(const FPropertyAnimatorCoreData&
 
 bool UPropertyAnimatorCoreIntHandler::SubtractValue(const FPropertyAnimatorCoreData& InPropertyData, const FInstancedPropertyBag& InValue)
 {
-	const FName PropertyName(InPropertyData.GetPathHash());
-	const TValueOrError<int32, EPropertyBagResult> ValueResult = InValue.GetValueInt32(PropertyName);
+	const FName PropertyHash(InPropertyData.GetLocatorPathHash());
+	const TValueOrError<int32, EPropertyBagResult> ValueResult = InValue.GetValueInt32(PropertyHash);
 	if (!ValueResult.HasValue())
 	{
 		return false;
diff --git a/Engine/Plugins/Experimental/PropertyAnimatorCore/Source/PropertyAnimatorCore/Private/Properties/Handlers/PropertyAnimatorCoreRotatorHandler.cpp b/Engine/Plugins/Experimental/PropertyAnimatorCore/Source/PropertyAnimatorCore/Private/Properties/Handlers/PropertyAnimatorCoreRotatorHandler.cpp
--- a/Engine/Plugins/Experimental/PropertyAnimatorCore/Source/PropertyAnimatorCore/Private/Properties/Handlers/PropertyAnimatorCoreRotatorHandler.cpp
+++ b/Engine/Plugins/Experimental/PropertyAnimatorCore/Source/PropertyAnimatorCore/Private/Properties/Handlers/PropertyAnimatorCoreRotatorHandler.cpp
@@ -14,28 +14,28 @@ bool UPropertyAnimatorCoreRotatorHandler::IsPropertySupported(const FPropertyAni
 
 bool UPropertyAnimatorCoreRotatorHandler::GetValue(const FPropertyAnimatorCoreData& InPropertyData, FInstancedPropertyBag& OutValue)
 {
-	const FName PropertyName(InPropertyData.GetPathHash());
-	OutValue.AddProperty(PropertyName, EPropertyBagPropertyType::Struct, TBaseStructure<FRotator>::Get());
+	const FName PropertyHash(InPropertyData.GetLocatorPathHash());
+	OutValue.AddProperty(PropertyHash, EPropertyBagPropertyType::Struct, TBaseStructure<FRotator>::Get());
 
 	FRotator Value;
 	InPropertyData.GetPropertyValuePtr(&Value);
 
-	OutValue.SetValueStruct(PropertyName, Value);
+	OutValue.SetValueStruct(PropertyHash, Value);
 
 	return true;
 }
 
 bool UPropertyAnimatorCoreRotatorHandler::SetValue(const FPropertyAnimatorCoreData& InPropertyData, const FInstancedPropertyBag& InValue)
 {
-	const FName PropertyName(InPropertyData.GetPathHash());
-	TValueOrError<FRotator*, EPropertyBagResult> ValueResult = InValue.GetValueStruct<FRotator>(PropertyName);
+	const FName PropertyHash(InPropertyData.GetLocatorPathHash());
+	TValueOrError<FRotator*, EPropertyBagResult> ValueResult = InValue.GetValueStruct<FRotator>(PropertyHash);
 	if (!ValueResult.HasValue())
 	{
 		return false;
 	}
 
-	FRotator NewValue = *ValueResult.GetValue();
-	InPropertyData.SetPropertyValuePtr(&NewValue);
+	FRotator* NewValue = ValueResult.GetValue();
+	InPropertyData.SetPropertyValuePtr(NewValue);
 
 	return true;
 }
@@ -47,8 +47,8 @@ bool UPropertyAnimatorCoreRotatorHandler::IsAdditiveSupported() const
 
 bool UPropertyAnimatorCoreRotatorHandler::AddValue(const FPropertyAnimatorCoreData& InPropertyData, const FInstancedPropertyBag& InValue)
 {
-	const FName PropertyName(InPropertyData.GetPathHash());
-	const TValueOrError<FRotator*, EPropertyBagResult> ValueResult = InValue.GetValueStruct<FRotator>(PropertyName);
+	const FName PropertyHash(InPropertyData.GetLocatorPathHash());
+	const TValueOrError<FRotator*, EPropertyBagResult> ValueResult = InValue.GetValueStruct<FRotator>(PropertyHash);
 	if (!ValueResult.HasValue())
 	{
 		return false;
@@ -65,8 +65,8 @@ bool UPropertyAnimatorCoreRotatorHandler::AddValue(const FPropertyAnimatorCoreDa
 
 bool UPropertyAnimatorCoreRotatorHandler::SubtractValue(const FPropertyAnimatorCoreData& InPropertyData, const FInstancedPropertyBag& InValue)
 {
-	const FName PropertyName(InPropertyData.GetPathHash());
-	const TValueOrError<FRotator*, EPropertyBagResult> ValueResult = InValue.GetValueStruct<FRotator>(PropertyName);
+	const FName PropertyHash(InPropertyData.GetLocatorPathHash());
+	const TValueOrError<FRotator*, EPropertyBagResult> ValueResult = InValue.GetValueStruct<FRotator>(PropertyHash);
 	if (!ValueResult.HasValue())
 	{
 		return false;
diff --git a/Engine/Plugins/Experimental/PropertyAnimatorCore/Source/PropertyAnimatorCore/Private/Properties/Handlers/PropertyAnimatorCoreStringHandler.cpp b/Engine/Plugins/Experimental/PropertyAnimatorCore/Source/PropertyAnimatorCore/Private/Properties/Handlers/PropertyAnimatorCoreStringHandler.cpp
--- a/Engine/Plugins/Experimental/PropertyAnimatorCore/Source/PropertyAnimatorCore/Private/Properties/Handlers/PropertyAnimatorCoreStringHandler.cpp
+++ b/Engine/Plugins/Experimental/PropertyAnimatorCore/Source/PropertyAnimatorCore/Private/Properties/Handlers/PropertyAnimatorCoreStringHandler.cpp
@@ -14,27 +14,27 @@ bool UPropertyAnimatorCoreStringHandler::IsPropertySupported(const FPropertyAnim
 
 bool UPropertyAnimatorCoreStringHandler::GetValue(const FPropertyAnimatorCoreData& InPropertyData, FInstancedPropertyBag& OutValue)
 {
-	const FName PropertyName(InPropertyData.GetPathHash());
-	OutValue.AddProperty(PropertyName, EPropertyBagPropertyType::String);
+	const FName PropertyHash(InPropertyData.GetLocatorPathHash());
+	OutValue.AddProperty(PropertyHash, EPropertyBagPropertyType::String);
 
 	FString Value;
 	InPropertyData.GetPropertyValuePtr(&Value);
 
-	OutValue.SetValueString(PropertyName, Value);
+	OutValue.SetValueString(PropertyHash, Value);
 
 	return true;
 }
 
 bool UPropertyAnimatorCoreStringHandler::SetValue(const FPropertyAnimatorCoreData& InPropertyData, const FInstancedPropertyBag& InValue)
 {
-	const FName PropertyName(InPropertyData.GetPathHash());
-	TValueOrError<FString, EPropertyBagResult> ValueResult = InValue.GetValueString(PropertyName);
+	const FName PropertyHash(InPropertyData.GetLocatorPathHash());
+	TValueOrError<FString, EPropertyBagResult> ValueResult = InValue.GetValueString(PropertyHash);
 	if (!ValueResult.HasValue())
 	{
 		return false;
 	}
 
-	FString NewValue = ValueResult.GetValue();
+	FString& NewValue = ValueResult.GetValue();
 	InPropertyData.SetPropertyValuePtr(&NewValue);
 
 	return true;
@@ -47,8 +47,8 @@ bool UPropertyAnimatorCoreStringHandler::IsAdditiveSupported() const
 
 bool UPropertyAnimatorCoreStringHandler::AddValue(const FPropertyAnimatorCoreData& InPropertyData, const FInstancedPropertyBag& InValue)
 {
-	const FName PropertyName(InPropertyData.GetPathHash());
-	const TValueOrError<FString, EPropertyBagResult> ValueResult = InValue.GetValueString(PropertyName);
+	const FName PropertyHash(InPropertyData.GetLocatorPathHash());
+	const TValueOrError<FString, EPropertyBagResult> ValueResult = InValue.GetValueString(PropertyHash);
 	if (!ValueResult.HasValue())
 	{
 		return false;
@@ -66,8 +66,8 @@ bool UPropertyAnimatorCoreStringHandler::AddValue(const FPropertyAnimatorCoreDat
 
 bool UPropertyAnimatorCoreStringHandler::SubtractValue(const FPropertyAnimatorCoreData& InPropertyData, const FInstancedPropertyBag& InValue)
 {
-	const FName PropertyName(InPropertyData.GetPathHash());
-	const TValueOrError<FString, EPropertyBagResult> ValueResult = InValue.GetValueString(PropertyName);
+	const FName PropertyHash(InPropertyData.GetLocatorPathHash());
+	const TValueOrError<FString, EPropertyBagResult> ValueResult = InValue.GetValueString(PropertyHash);
 	if (!ValueResult.HasValue())
 	{
 		return false;
diff --git a/Engine/Plugins/Experimental/PropertyAnimatorCore/Source/PropertyAnimatorCore/Private/Properties/Handlers/PropertyAnimatorCoreTextHandler.cpp b/Engine/Plugins/Experimental/PropertyAnimatorCore/Source/PropertyAnimatorCore/Private/Properties/Handlers/PropertyAnimatorCoreTextHandler.cpp
--- a/Engine/Plugins/Experimental/PropertyAnimatorCore/Source/PropertyAnimatorCore/Private/Properties/Handlers/PropertyAnimatorCoreTextHandler.cpp
+++ b/Engine/Plugins/Experimental/PropertyAnimatorCore/Source/PropertyAnimatorCore/Private/Properties/Handlers/PropertyAnimatorCoreTextHandler.cpp
@@ -16,27 +16,27 @@ bool UPropertyAnimatorCoreTextHandler::IsPropertySupported(const FPropertyAnimat
 
 bool UPropertyAnimatorCoreTextHandler::GetValue(const FPropertyAnimatorCoreData& InPropertyData, FInstancedPropertyBag& OutValue)
 {
-	const FName PropertyName(InPropertyData.GetPathHash());
-	OutValue.AddProperty(PropertyName, EPropertyBagPropertyType::Text);
+	const FName PropertyHash(InPropertyData.GetLocatorPathHash());
+	OutValue.AddProperty(PropertyHash, EPropertyBagPropertyType::Text);
 
 	FText Value;
 	InPropertyData.GetPropertyValuePtr(&Value);
 
-	OutValue.SetValueText(PropertyName, Value);
+	OutValue.SetValueText(PropertyHash, Value);
 
 	return true;
 }
 
 bool UPropertyAnimatorCoreTextHandler::SetValue(const FPropertyAnimatorCoreData& InPropertyData, const FInstancedPropertyBag& InValue)
 {
-	const FName PropertyName(InPropertyData.GetPathHash());
-	TValueOrError<FText, EPropertyBagResult> ValueResult = InValue.GetValueText(PropertyName);
+	const FName PropertyHash(InPropertyData.GetLocatorPathHash());
+	TValueOrError<FText, EPropertyBagResult> ValueResult = InValue.GetValueText(PropertyHash);
 	if (!ValueResult.HasValue())
 	{
 		return false;
 	}
 
-	FText NewValue = ValueResult.GetValue();
+	FText& NewValue = ValueResult.GetValue();
 	InPropertyData.SetPropertyValuePtr(&NewValue);
 
 	return true;
@@ -49,8 +49,8 @@ bool UPropertyAnimatorCoreTextHandler::IsAdditiveSupported() const
 
 bool UPropertyAnimatorCoreTextHandler::AddValue(const FPropertyAnimatorCoreData& InPropertyData, const FInstancedPropertyBag& InValue)
 {
-	const FName PropertyName(InPropertyData.GetPathHash());
-	const TValueOrError<FText, EPropertyBagResult> ValueResult = InValue.GetValueText(PropertyName);
+	const FName PropertyHash(InPropertyData.GetLocatorPathHash());
+	const TValueOrError<FText, EPropertyBagResult> ValueResult = InValue.GetValueText(PropertyHash);
 	if (!ValueResult.HasValue())
 	{
 		return false;
@@ -68,8 +68,8 @@ bool UPropertyAnimatorCoreTextHandler::AddValue(const FPropertyAnimatorCoreData&
 
 bool UPropertyAnimatorCoreTextHandler::SubtractValue(const FPropertyAnimatorCoreData& InPropertyData, const FInstancedPropertyBag& InValue)
 {
-	const FName PropertyName(InPropertyData.GetPathHash());
-	const TValueOrError<FText, EPropertyBagResult> ValueResult = InValue.GetValueText(PropertyName);
+	const FName PropertyHash(InPropertyData.GetLocatorPathHash());
+	const TValueOrError<FText, EPropertyBagResult> ValueResult = InValue.GetValueText(PropertyHash);
 	if (!ValueResult.HasValue())
 	{
 		return false;
diff --git a/Engine/Plugins/Experimental/PropertyAnimatorCore/Source/PropertyAnimatorCore/Private/Properties/Handlers/PropertyAnimatorCoreVectorHandler.cpp b/Engine/Plugins/Experimental/PropertyAnimatorCore/Source/PropertyAnimatorCore/Private/Properties/Handlers/PropertyAnimatorCoreVectorHandler.cpp
--- a/Engine/Plugins/Experimental/PropertyAnimatorCore/Source/PropertyAnimatorCore/Private/Properties/Handlers/PropertyAnimatorCoreVectorHandler.cpp
+++ b/Engine/Plugins/Experimental/PropertyAnimatorCore/Source/PropertyAnimatorCore/Private/Properties/Handlers/PropertyAnimatorCoreVectorHandler.cpp
@@ -14,28 +14,28 @@ bool UPropertyAnimatorCoreVectorHandler::IsPropertySupported(const FPropertyAnim
 
 bool UPropertyAnimatorCoreVectorHandler::GetValue(const FPropertyAnimatorCoreData& InPropertyData, FInstancedPropertyBag& OutValue)
 {
-	const FName PropertyName(InPropertyData.GetPathHash());
-	OutValue.AddProperty(PropertyName, EPropertyBagPropertyType::Struct, TBaseStructure<FVector>::Get());
+	const FName PropertyHash(InPropertyData.GetLocatorPathHash());
+	OutValue.AddProperty(PropertyHash, EPropertyBagPropertyType::Struct, TBaseStructure<FVector>::Get());
 
 	FVector Value;
 	InPropertyData.GetPropertyValuePtr(&Value);
 
-	OutValue.SetValueStruct(PropertyName, Value);
+	OutValue.SetValueStruct(PropertyHash, Value);
 
 	return true;
 }
 
 bool UPropertyAnimatorCoreVectorHandler::SetValue(const FPropertyAnimatorCoreData& InPropertyData, const FInstancedPropertyBag& InValue)
 {
-	const FName PropertyName(InPropertyData.GetPathHash());
-	TValueOrError<FVector*, EPropertyBagResult> ValueResult = InValue.GetValueStruct<FVector>(PropertyName);
+	const FName PropertyHash(InPropertyData.GetLocatorPathHash());
+	TValueOrError<FVector*, EPropertyBagResult> ValueResult = InValue.GetValueStruct<FVector>(PropertyHash);
 	if (!ValueResult.HasValue())
 	{
 		return false;
 	}
 
-	FVector NewValue = *ValueResult.GetValue();
-	InPropertyData.SetPropertyValuePtr(&NewValue);
+	FVector* NewValue = ValueResult.GetValue();
+	InPropertyData.SetPropertyValuePtr(NewValue);
 
 	return true;
 }
@@ -47,8 +47,8 @@ bool UPropertyAnimatorCoreVectorHandler::IsAdditiveSupported() const
 
 bool UPropertyAnimatorCoreVectorHandler::AddValue(const FPropertyAnimatorCoreData& InPropertyData, const FInstancedPropertyBag& InValue)
 {
-	const FName PropertyName(InPropertyData.GetPathHash());
-	const TValueOrError<FVector*, EPropertyBagResult> ValueResult = InValue.GetValueStruct<FVector>(PropertyName);
+	const FName PropertyHash(InPropertyData.GetLocatorPathHash());
+	const TValueOrError<FVector*, EPropertyBagResult> ValueResult = InValue.GetValueStruct<FVector>(PropertyHash);
 	if (!ValueResult.HasValue())
 	{
 		return false;
@@ -65,8 +65,8 @@ bool UPropertyAnimatorCoreVectorHandler::AddValue(const FPropertyAnimatorCoreDat
 
 bool UPropertyAnimatorCoreVectorHandler::SubtractValue(const FPropertyAnimatorCoreData& InPropertyData, const FInstancedPropertyBag& InValue)
 {
-	const FName PropertyName(InPropertyData.GetPathHash());
-	const TValueOrError<FVector*, EPropertyBagResult> ValueResult = InValue.GetValueStruct<FVector>(PropertyName);
+	const FName PropertyHash(InPropertyData.GetLocatorPathHash());
+	const TValueOrError<FVector*, EPropertyBagResult> ValueResult = InValue.GetValueStruct<FVector>(PropertyHash);
 	if (!ValueResult.HasValue())
 	{
 		return false;
diff --git a/Engine/Plugins/Experimental/PropertyAnimatorCore/Source/PropertyAnimatorCore/Private/Properties/PropertyAnimatorCoreContext.cpp b/Engine/Plugins/Experimental/PropertyAnimatorCore/Source/PropertyAnimatorCore/Private/Properties/PropertyAnimatorCoreContext.cpp
--- a/Engine/Plugins/Experimental/PropertyAnimatorCore/Source/PropertyAnimatorCore/Private/Properties/PropertyAnimatorCoreContext.cpp
+++ b/Engine/Plugins/Experimental/PropertyAnimatorCore/Source/PropertyAnimatorCore/Private/Properties/PropertyAnimatorCoreContext.cpp
@@ -125,19 +125,39 @@ void UPropertyAnimatorCoreContext::PostLoad()
 	CheckEditConverterRule();
 	CheckEditResolver();
 
-	FTSTicker::GetCoreTicker().AddTicker(FTickerDelegate::CreateWeakLambda(this, [this](float InDelta)
+	AnimatedProperty.GeneratePropertyPath();
+
+	if (AnimatedProperty.IsResolvable())
 	{
-		// Restore before regenerating new property path
-		Restore();
+		DeltaPropertyValues.Reset();
+		OriginalPropertyValues.Reset();
+	}
+
+	for (const FPropertyAnimatorCoreData& Property : ResolveProperty(/** Evaluation */false))
+	{
+		PRAGMA_DISABLE_DEPRECATION_WARNINGS
+		const FName OldPropertyPath(Property.GetPathHash());
+		PRAGMA_ENABLE_DEPRECATION_WARNINGS
+
+		const FName NewPropertyPath(Property.GetLocatorPathHash());
 
-		if (AnimatedProperty.IsResolvable())
+		if (DeltaPropertyValues.RenameProperty(OldPropertyPath, NewPropertyPath) == EPropertyBagAlterationResult::Success)
 		{
-			DeltaPropertyValues.Reset();
-			OriginalPropertyValues.Reset();
+			UE_LOG(LogPropertyAnimatorCore, Log, TEXT("Property %s in delta property bag migrated successfully %s to %s"), *Property.GetPropertyDisplayName(), *OldPropertyPath.ToString(), *NewPropertyPath.ToString());
 		}
 
-		AnimatedProperty.GeneratePropertyPath();
+		if (OriginalPropertyValues.RenameProperty(OldPropertyPath, NewPropertyPath) == EPropertyBagAlterationResult::Success)
+		{
+			UE_LOG(LogPropertyAnimatorCore, Log, TEXT("Property %s in absolute property bag migrated successfully %s to %s"), *Property.GetPropertyDisplayName(), *OldPropertyPath.ToString(), *NewPropertyPath.ToString());
+		}
+	}
 
+	/** Restore property values on next tick */
+	FTSTicker::GetCoreTicker().AddTicker(FTickerDelegate::CreateWeakLambda(this, [this](float)
+	{
+		Restore();
+
+		/** Stop ticker */
 		return false;
 	}));
 }
@@ -239,7 +259,7 @@ bool UPropertyAnimatorCoreContext::ExportPreset(const UPropertyAnimatorCorePrese
 		ContextArchive->Set(GET_MEMBER_NAME_STRING_CHECKED(UPropertyAnimatorCoreContext, Mode), static_cast<uint64>(Mode));
 	}
 
-	ContextArchive->Set(GET_MEMBER_NAME_STRING_CHECKED(UPropertyAnimatorCoreContext, AnimatedProperty), AnimatedProperty.GetPropertyLocatorPath());
+	ContextArchive->Set(GET_MEMBER_NAME_STRING_CHECKED(UPropertyAnimatorCoreContext, AnimatedProperty), AnimatedProperty.GetLocatorPath());
 
 	if (Resolver)
 	{
@@ -293,7 +313,7 @@ bool UPropertyAnimatorCoreContext::ResolvePropertyOwner(AActor* InNewOwner)
 	const bool bFound = IsValid(NewOwningActor);
 
 	// Try to resolve property owner on new owning actor
-	UObject* NewOwner = FPropertyAnimatorCoreData(NewOwningActor, AnimatedProperty.GetPropertyLocatorPath()).GetOwner();
+	UObject* NewOwner = FPropertyAnimatorCoreData(NewOwningActor, AnimatedProperty.GetLocatorPath()).GetOwner();
 
 	const FProperty* MemberProperty = AnimatedProperty.GetMemberProperty();
 	UClass* PropertyOwningClass = MemberProperty->GetOwnerClass();
@@ -307,7 +327,7 @@ bool UPropertyAnimatorCoreContext::ResolvePropertyOwner(AActor* InNewOwner)
 		return true;
 	}
 
-	UE_LOG(LogPropertyAnimatorCore, Warning, TEXT("Could not resolve property owner %s on actor %s"), *AnimatedProperty.GetPathHash(), NewOwningActor ? *NewOwningActor->GetActorNameOrLabel() : TEXT("Invalid"))
+	UE_LOG(LogPropertyAnimatorCore, Warning, TEXT("Could not resolve property %s on actor %s"), *AnimatedProperty.GetLocatorPath(), NewOwningActor ? *NewOwningActor->GetActorNameOrLabel() : TEXT("Invalid"))
 
 	return false;
 }
@@ -391,7 +411,7 @@ void UPropertyAnimatorCoreContext::Restore()
 			// Reset original value
 			if (Handler->SetValue(ResolvedProperty, OriginalPropertyValues))
 			{
-				OriginalPropertyValues.RemovePropertyByName(FName(ResolvedProperty.GetPathHash()));
+				OriginalPropertyValues.RemovePropertyByName(ResolvedProperty.GetLocatorPathHash());
 			}
 		}
 
@@ -404,7 +424,7 @@ void UPropertyAnimatorCoreContext::Restore()
 			// Subtract delta value
 			if (Handler->SubtractValue(ResolvedProperty, DeltaPropertyValues))
 			{
-				DeltaPropertyValues.RemovePropertyByName(FName(ResolvedProperty.GetPathHash()));
+				DeltaPropertyValues.RemovePropertyByName(ResolvedProperty.GetLocatorPathHash());
 			}
 		}
 	}
@@ -423,7 +443,7 @@ void UPropertyAnimatorCoreContext::Save()
 
 	for (const FPropertyAnimatorCoreData& PropertyData : ResolveProperty(false))
 	{
-		const FName Name(PropertyData.GetPathHash());
+		const FName Name(PropertyData.GetLocatorPathHash());
 		if (!OriginalPropertyValues.FindPropertyDescByName(Name))
 		{
 			const FProperty* Property = PropertyData.GetLeafProperty();
@@ -463,11 +483,13 @@ void UPropertyAnimatorCoreContext::CommitEvaluationResult(const FPropertyAnimato
 		return;
 	}
 
-	const FName PropertyName(InResolvedProperty.GetPathHash());
+	const FName PropertyName(InResolvedProperty.GetLocatorPathHash());
 
 	const FPropertyBagPropertyDesc* FromProperty = InEvaluatedValues.FindPropertyDescByName(PropertyName);
 	const FPropertyBagPropertyDesc* ToProperty = DeltaPropertyValues.FindPropertyDescByName(PropertyName);
 
+	checkf(!!FromProperty && !!ToProperty, TEXT("Property bags are missing the evaluated property"));
+
 	if (UPropertyAnimatorCoreConverterBase* Converter = ConverterClass.GetDefaultObject())
 	{
 		if (!Converter->Convert(*FromProperty, InEvaluatedValues, *ToProperty, DeltaPropertyValues, ConverterRule.IsValid() ? &ConverterRule : nullptr))
diff --git a/Engine/Plugins/Experimental/PropertyAnimatorCore/Source/PropertyAnimatorCore/Private/Properties/PropertyAnimatorCoreData.cpp b/Engine/Plugins/Experimental/PropertyAnimatorCore/Source/PropertyAnimatorCore/Private/Properties/PropertyAnimatorCoreData.cpp
--- a/Engine/Plugins/Experimental/PropertyAnimatorCore/Source/PropertyAnimatorCore/Private/Properties/PropertyAnimatorCoreData.cpp
+++ b/Engine/Plugins/Experimental/PropertyAnimatorCore/Source/PropertyAnimatorCore/Private/Properties/PropertyAnimatorCoreData.cpp
@@ -297,38 +297,78 @@ TArray<UObject*> FPropertyAnimatorCoreData::GetOuters(const UObject* InStopOuter
 	return Owners;
 }
 
-TArray<FString> FPropertyAnimatorCoreData::GetOuterNames() const
+FString FPropertyAnimatorCoreData::GetPropertyDisplayName() const
 {
-	TArray<FString> OuterNames;
-	PathHash.ParseIntoArray(OuterNames, TEXT("."), true);
-	const FString MemberPropertyName = GetMemberPropertyName().ToString();
-
-	// Remove all properties chain segments
-	for (int32 Index = OuterNames.Num() - 1; Index >= 0; Index--)
+	if (!PropertyDisplayName.IsEmpty())
 	{
-		if (OuterNames[Index].Equals(MemberPropertyName))
-		{
-			OuterNames.RemoveAt(Index, OuterNames.Num() - Index);
-			break;
-		}
+		return PropertyDisplayName;
 	}
 
-	// Remove resolver name at the beginning
-	if (!OuterNames.IsEmpty() && IsResolvable())
+	FPropertyAnimatorCoreData* MutableThis = const_cast<FPropertyAnimatorCoreData*>(this);
+	MutableThis->GeneratePropertyDisplayName();
+
+	return PropertyDisplayName;
+}
+
+FString FPropertyAnimatorCoreData::GetPathHash() const
+{
+	PRAGMA_DISABLE_DEPRECATION_WARNINGS
+	FString PathHashValue = PathHash;
+	PRAGMA_ENABLE_DEPRECATION_WARNINGS
+	return PathHashValue;
+}
+
+FString FPropertyAnimatorCoreData::GetLocatorPath() const
+{
+	if (!LocatorPath.IsEmpty())
 	{
-		OuterNames.RemoveAt(0);
+		return LocatorPath;
 	}
 
-	return OuterNames;
+	FPropertyAnimatorCoreData* MutableThis = const_cast<FPropertyAnimatorCoreData*>(this);
+	MutableThis->GeneratePropertyLocatorPath();
+
+	return LocatorPath;
 }
 
-FString FPropertyAnimatorCoreData::GetPropertyDisplayName() const
+FName FPropertyAnimatorCoreData::GetLocatorPathHash() const
 {
-	if (!PropertyDisplayName.IsEmpty())
+	return FName(LexToString(GetTypeHash(GetLocatorPath())));
+}
+
+void FPropertyAnimatorCoreData::GeneratePropertyLocatorPath()
+{
+	FString NewLocatorPath;
+
+	// Append resolver
+	if (const UPropertyAnimatorCoreResolver* Resolver = GetPropertyResolver())
 	{
-		return PropertyDisplayName;
+		NewLocatorPath += Resolver->GetClass()->GetClassPathName().ToString() + TEXT(":") + Resolver->GetResolverName().ToString() + TEXT(";");
+	}
+
+	// Append outers
+	const AActor* OwningActor = GetOwningActor();
+	for (const UObject* Outer : GetOuters(OwningActor))
+	{
+		NewLocatorPath += Outer->GetClass()->GetClassPathName().ToString() + TEXT(":") + Outer->GetName() + TEXT(",");
+	}
+
+	NewLocatorPath.RemoveFromEnd(TEXT(","));
+	NewLocatorPath += TEXT(";");
+
+	// Append properties
+	for (const TFieldPath<FProperty>& Property : ChainProperties)
+	{
+		NewLocatorPath += Property->GetPathName(OwningActor) + TEXT(",");
 	}
 
+	NewLocatorPath.RemoveFromEnd(TEXT(","));
+
+	LocatorPath = NewLocatorPath;
+}
+
+void FPropertyAnimatorCoreData::GeneratePropertyDisplayName()
+{
 	const FString ResolverName = IsResolvable()
 		? (GetPropertyResolver()->GetResolverName().ToString() + TEXT("."))
 		: TEXT("");
@@ -364,8 +404,6 @@ FString FPropertyAnimatorCoreData::GetPropertyDisplayName() const
 			: (TEXT(".") + LeafPropertyName);
 	}
 
-	FPropertyAnimatorCoreData* MutableThis = const_cast<FPropertyAnimatorCoreData*>(this);
-
 	// Find alias
 	FString AliasName;
 	if (const UPropertyAnimatorCoreSubsystem* AnimatorSubsystem = UPropertyAnimatorCoreSubsystem::Get())
@@ -384,40 +422,7 @@ FString FPropertyAnimatorCoreData::GetPropertyDisplayName() const
 		DisplayName += AliasName;
 	}
 
-	MutableThis->PropertyDisplayName = ResolverName + DisplayName;
-
-	return PropertyDisplayName;
-}
-
-FString FPropertyAnimatorCoreData::GetPropertyLocatorPath() const
-{
-	FString LocatorPath;
-
-	// Append resolver
-	if (const UPropertyAnimatorCoreResolver* Resolver = GetPropertyResolver())
-	{
-		LocatorPath += Resolver->GetClass()->GetClassPathName().ToString() + TEXT(":") + Resolver->GetResolverName().ToString() + TEXT(";");
-	}
-
-	// Append outers
-	for (const UObject* Outer : GetOuters(GetOwningActor()))
-	{
-		LocatorPath += Outer->GetClass()->GetClassPathName().ToString() + TEXT(":") + Outer->GetName() + TEXT(",");
-	}
-
-	LocatorPath.RemoveFromEnd(TEXT(","));
-	LocatorPath += TEXT(";");
-
-	// Append properties
-	const AActor* OwningActor = GetOwningActor();
-	for (const TFieldPath<FProperty>& Property : ChainProperties)
-	{
-		LocatorPath += Property->GetPathName(OwningActor) + TEXT(",");
-	}
-
-	LocatorPath.RemoveFromEnd(TEXT(","));
-
-	return LocatorPath;
+	PropertyDisplayName = ResolverName + DisplayName;
 }
 
 FName FPropertyAnimatorCoreData::GetMemberPropertyName() const
@@ -818,6 +823,8 @@ FName FPropertyAnimatorCoreData::GetPropertyTypeName(const FProperty* InProperty
 
 void FPropertyAnimatorCoreData::GeneratePropertyPath()
 {
+	PRAGMA_DISABLE_DEPRECATION_WARNINGS
+
 	const UObject* Owner = GetOwner();
 	const UObject* StopOuter = GetOwningActor();
 	const FString ResolverName = IsResolvable() ? GetPropertyResolver()->GetResolverName().ToString() + TEXT(".") : TEXT("");
@@ -830,6 +837,11 @@ void FPropertyAnimatorCoreData::GeneratePropertyPath()
 	{
 		PathHash += TEXT(".") + ChainProperty->GetName();
 	}
+
+	PRAGMA_ENABLE_DEPRECATION_WARNINGS
+
+	GeneratePropertyLocatorPath();
+	GeneratePropertyDisplayName();
 }
 
 bool FPropertyAnimatorCoreData::FindSetterFunctions()
diff --git a/Engine/Plugins/Experimental/PropertyAnimatorCore/Source/PropertyAnimatorCore/Public/Properties/PropertyAnimatorCoreData.h b/Engine/Plugins/Experimental/PropertyAnimatorCore/Source/PropertyAnimatorCore/Public/Properties/PropertyAnimatorCoreData.h
--- a/Engine/Plugins/Experimental/PropertyAnimatorCore/Source/PropertyAnimatorCore/Public/Properties/PropertyAnimatorCoreData.h
+++ b/Engine/Plugins/Experimental/PropertyAnimatorCore/Source/PropertyAnimatorCore/Public/Properties/PropertyAnimatorCoreData.h
@@ -55,12 +55,12 @@ struct FPropertyAnimatorCoreData
 
 	friend uint32 GetTypeHash(const FPropertyAnimatorCoreData& InItem)
 	{
-		return GetTypeHash(InItem.PathHash);
+		return GetTypeHash(InItem.LocatorPath);
 	}
 
 	bool operator==(const FPropertyAnimatorCoreData& InOther) const
 	{
-		return PathHash == InOther.PathHash;
+		return LocatorPath == InOther.LocatorPath;
 	}
 
 	/** Returns the owning actor from this owner object */
@@ -83,8 +83,6 @@ struct FPropertyAnimatorCoreData
 	/** Returns chain of owner until it reaches StopOuter */
 	TArray<UObject*> GetOuters(const UObject* InStopOuter) const;
 
-	TArray<FString> GetOuterNames() const;
-
 	/** The member property of the owner, top property inside the owner itself */
 	FProperty* GetMemberProperty() const
 	{
@@ -100,14 +98,18 @@ struct FPropertyAnimatorCoreData
 	/** The friendly display name from the member property to the inner property */
 	PROPERTYANIMATORCORE_API FString GetPropertyDisplayName() const;
 
-	/** The full path of a property with its owner */
-	FString GetPathHash() const
-	{
-		return PathHash;
-	}
+	/**
+	 * DEPRECATED, use GetLocatorPath instead
+	 * The full path of a property with its owner
+	 */
+	UE_DEPRECATED(5.6, "Use GetLocatorPath instead")
+	PROPERTYANIMATORCORE_API FString GetPathHash() const;
 
-	/** Creates a single path that contains resolver, owner and properties for (de)serialization */
-	PROPERTYANIMATORCORE_API FString GetPropertyLocatorPath() const;
+	/** The path to retrieve owner and property on any actor */
+	PROPERTYANIMATORCORE_API FString GetLocatorPath() const;
+
+	/** Hash used to store property, based on the locator path */
+	PROPERTYANIMATORCORE_API FName GetLocatorPathHash() const;
 
 	/** The member property name */
 	PROPERTYANIMATORCORE_API FName GetMemberPropertyName() const;
@@ -201,7 +203,7 @@ struct FPropertyAnimatorCoreData
 	/** Gets the property handler to perform operation on property without knowing the type */
 	UPropertyAnimatorCoreHandlerBase* GetPropertyHandler() const;
 
-	/** Internal use only, create the hash and display name */
+	/** Internal use only, (re)creates the property path, locator path, display name */
 	void GeneratePropertyPath();
 
 private:
@@ -227,6 +229,12 @@ private:
 	/** Tries to find setter function for this property */
 	bool FindSetterFunctions();
 
+	/** Creates a single path that contains resolver, owner and properties for (de)serialization */
+	void GeneratePropertyLocatorPath();
+
+	/** Creates the display name for the underlying property */
+	void GeneratePropertyDisplayName();
+
 	/** Owner of the property */
 	UPROPERTY()
 	TWeakObjectPtr<UObject> OwnerWeak;
@@ -235,10 +243,15 @@ private:
 	UPROPERTY(Transient)
 	FString PropertyDisplayName;
 
-	/** Used to quickly compare struct of this type */
+	/** Previously used to quickly compare struct of this type, replaced by LocatorPath */
+	UE_DEPRECATED(5.6, "Use LocatorPath instead")
 	UPROPERTY()
 	FString PathHash;
 
+	/** Property locator path to find the property and owner on any actor */
+	UPROPERTY()
+	FString LocatorPath;
+
 	/** Chain from member property to inner property */
 	UPROPERTY()
 	TArray<TFieldPath<FProperty>> ChainProperties;
diff --git a/Engine/Plugins/Experimental/PropertyAnimatorCore/Source/PropertyAnimatorCoreEditor/Private/Menus/PropertyAnimatorCoreEditorMenu.cpp b/Engine/Plugins/Experimental/PropertyAnimatorCore/Source/PropertyAnimatorCoreEditor/Private/Menus/PropertyAnimatorCoreEditorMenu.cpp
--- a/Engine/Plugins/Experimental/PropertyAnimatorCore/Source/PropertyAnimatorCoreEditor/Private/Menus/PropertyAnimatorCoreEditorMenu.cpp
+++ b/Engine/Plugins/Experimental/PropertyAnimatorCore/Source/PropertyAnimatorCoreEditor/Private/Menus/PropertyAnimatorCoreEditorMenu.cpp
@@ -131,7 +131,7 @@ void UE::PropertyAnimatorCoreEditor::Menu::FillExistingAnimatorSection(UToolMenu
 	{
 		for (AActor* Actor : InMenuData->GetContext().GetActors())
 		{
-			const FPropertyAnimatorCoreData ActorProperty(Actor, Property.GetPropertyLocatorPath());
+			const FPropertyAnimatorCoreData ActorProperty(Actor, Property.GetLocatorPath());
 			for (UPropertyAnimatorCoreBase* Animator : Subsystem->GetExistingAnimators(ActorProperty))
 			{
 				const FName AnimatorKey = FName(Animator->GetAnimatorDisplayName().ToString() + TEXT(" (") + Animator->GetClass()->GetName() + TEXT(")"));
@@ -866,7 +866,7 @@ void UE::PropertyAnimatorCoreEditor::Menu::FillNewAnimatorSubmenu(UToolMenu* InM
 		{
 			const FString MenuName = SupportedProperty.GetPropertyDisplayName();
 			const FText MenuLabel = FText::FromString(MenuName + TEXT(" (") + SupportedProperty.GetLeafPropertyTypeName().ToString() + TEXT(")"));
-			const FString PropertyLocatorPath = SupportedProperty.GetPropertyLocatorPath();
+			const FString PropertyLocatorPath = SupportedProperty.GetLocatorPath();
 
 			// Create action (creates an animator and links the property)
 			PropertySection.AddMenuEntry(
@@ -982,7 +982,7 @@ void UE::PropertyAnimatorCoreEditor::Menu::FillLinkAnimatorSubmenu(UToolMenu* In
 	{
 		const FString MenuName = SupportedProperty.GetPropertyDisplayName();
 		const FText MenuLabel = FText::FromString(MenuName + TEXT(" (") + SupportedProperty.GetLeafPropertyTypeName().ToString() + TEXT(")"));
-		const FString PropertyLocatorPath = SupportedProperty.GetPropertyLocatorPath();
+		const FString PropertyLocatorPath = SupportedProperty.GetLocatorPath();
 		
 		PropertySection.AddMenuEntry(
 			FName(MenuName)
@@ -1057,7 +1057,7 @@ void UE::PropertyAnimatorCoreEditor::Menu::FillPresetAnimatorSubmenu(UToolMenu*
 		const FString MenuName = SupportedProperty.GetPropertyDisplayName();
 		const FText MenuLabel = FText::FromString(MenuName + TEXT(" (") + SupportedProperty.GetLeafPropertyTypeName().ToString() + TEXT(")"));
 		const FText MenuTooltip = LOCTEXT("LinkPresetProperty.Tooltip", "Link this preset property");
-		const FString PropertyLocatorPath = SupportedProperty.GetPropertyLocatorPath();
+		const FString PropertyLocatorPath = SupportedProperty.GetLocatorPath();
 
 		const FToolMenuEntry SupportedPropertyEntry = FToolMenuEntry::InitMenuEntry(
 			FName(MenuName)
@@ -1148,7 +1148,7 @@ void UE::PropertyAnimatorCoreEditor::Menu::FillNewPresetAnimatorSubmenu(UToolMen
 	{
 		const FString MenuName = PresetProperty.GetPropertyDisplayName();
 		const FText MenuLabel = FText::FromString(MenuName + TEXT(" (") + PresetProperty.GetLeafPropertyTypeName().ToString() + TEXT(")"));
-		const FString PropertyLocatorPath = PresetProperty.GetPropertyLocatorPath();
+		const FString PropertyLocatorPath = PresetProperty.GetLocatorPath();
 		
 		const FToolMenuEntry CreatePropertyEntry = FToolMenuEntry::InitMenuEntry(
 			FName(TEXT("CreateProperty") + MenuName)
