MotionDesign : PropertyAnimators
- Fix horde deprecation warning about PathHash even with deprecation macro being used

diff --git a/Engine/Plugins/Experimental/PropertyAnimatorCore/Source/PropertyAnimatorCore/Private/Properties/PropertyAnimatorCoreContext.cpp b/Engine/Plugins/Experimental/PropertyAnimatorCore/Source/PropertyAnimatorCore/Private/Properties/PropertyAnimatorCoreContext.cpp
--- a/Engine/Plugins/Experimental/PropertyAnimatorCore/Source/PropertyAnimatorCore/Private/Properties/PropertyAnimatorCoreContext.cpp
+++ b/Engine/Plugins/Experimental/PropertyAnimatorCore/Source/PropertyAnimatorCore/Private/Properties/PropertyAnimatorCoreContext.cpp
@@ -135,10 +135,7 @@ void UPropertyAnimatorCoreContext::PostLoad()
 
 	for (const FPropertyAnimatorCoreData& Property : ResolveProperty(/** Evaluation */false))
 	{
-		PRAGMA_DISABLE_DEPRECATION_WARNINGS
 		const FName OldPropertyPath(Property.GetPathHash());
-		PRAGMA_ENABLE_DEPRECATION_WARNINGS
-
 		const FName NewPropertyPath(Property.GetLocatorPathHash());
 
 		if (DeltaPropertyValues.RenameProperty(OldPropertyPath, NewPropertyPath) == EPropertyBagAlterationResult::Success)
diff --git a/Engine/Plugins/Experimental/PropertyAnimatorCore/Source/PropertyAnimatorCore/Private/Properties/PropertyAnimatorCoreData.cpp b/Engine/Plugins/Experimental/PropertyAnimatorCore/Source/PropertyAnimatorCore/Private/Properties/PropertyAnimatorCoreData.cpp
--- a/Engine/Plugins/Experimental/PropertyAnimatorCore/Source/PropertyAnimatorCore/Private/Properties/PropertyAnimatorCoreData.cpp
+++ b/Engine/Plugins/Experimental/PropertyAnimatorCore/Source/PropertyAnimatorCore/Private/Properties/PropertyAnimatorCoreData.cpp
@@ -312,10 +312,7 @@ FString FPropertyAnimatorCoreData::GetPropertyDisplayName() const
 
 FString FPropertyAnimatorCoreData::GetPathHash() const
 {
-	PRAGMA_DISABLE_DEPRECATION_WARNINGS
-	FString PathHashValue = PathHash;
-	PRAGMA_ENABLE_DEPRECATION_WARNINGS
-	return PathHashValue;
+	return PathHash;
 }
 
 FString FPropertyAnimatorCoreData::GetLocatorPath() const
@@ -823,8 +820,6 @@ FName FPropertyAnimatorCoreData::GetPropertyTypeName(const FProperty* InProperty
 
 void FPropertyAnimatorCoreData::GeneratePropertyPath()
 {
-	PRAGMA_DISABLE_DEPRECATION_WARNINGS
-
 	const UObject* Owner = GetOwner();
 	const UObject* StopOuter = GetOwningActor();
 	const FString ResolverName = IsResolvable() ? GetPropertyResolver()->GetResolverName().ToString() + TEXT(".") : TEXT("");
@@ -838,8 +833,6 @@ void FPropertyAnimatorCoreData::GeneratePropertyPath()
 		PathHash += TEXT(".") + ChainProperty->GetName();
 	}
 
-	PRAGMA_ENABLE_DEPRECATION_WARNINGS
-
 	GeneratePropertyLocatorPath();
 	GeneratePropertyDisplayName();
 }
diff --git a/Engine/Plugins/Experimental/PropertyAnimatorCore/Source/PropertyAnimatorCore/Public/Properties/PropertyAnimatorCoreData.h b/Engine/Plugins/Experimental/PropertyAnimatorCore/Source/PropertyAnimatorCore/Public/Properties/PropertyAnimatorCoreData.h
--- a/Engine/Plugins/Experimental/PropertyAnimatorCore/Source/PropertyAnimatorCore/Public/Properties/PropertyAnimatorCoreData.h
+++ b/Engine/Plugins/Experimental/PropertyAnimatorCore/Source/PropertyAnimatorCore/Public/Properties/PropertyAnimatorCoreData.h
@@ -102,7 +102,6 @@ struct FPropertyAnimatorCoreData
 	 * DEPRECATED, use GetLocatorPath instead
 	 * The full path of a property with its owner
 	 */
-	UE_DEPRECATED(5.6, "Use GetLocatorPath instead")
 	PROPERTYANIMATORCORE_API FString GetPathHash() const;
 
 	/** The path to retrieve owner and property on any actor */
@@ -243,8 +242,9 @@ private:
 	UPROPERTY(Transient)
 	FString PropertyDisplayName;
 
-	/** Previously used to quickly compare struct of this type, replaced by LocatorPath */
-	UE_DEPRECATED(5.6, "Use LocatorPath instead")
+	/**
+	 * DEPRECATED, use locator path instead,
+	 * Previously used to quickly compare struct of this type, replaced by LocatorPath */
 	UPROPERTY()
 	FString PathHash;
 
