MotionDesign : PropertyAnimator
- Improving animator manual time source to allow custom time and player to be used together

diff --git a/Engine/Plugins/Experimental/PropertyAnimatorCore/Source/PropertyAnimatorCore/Private/TimeSources/PropertyAnimatorCoreManualTimeSource.cpp b/Engine/Plugins/Experimental/PropertyAnimatorCore/Source/PropertyAnimatorCore/Private/TimeSources/PropertyAnimatorCoreManualTimeSource.cpp
--- a/Engine/Plugins/Experimental/PropertyAnimatorCore/Source/PropertyAnimatorCore/Private/TimeSources/PropertyAnimatorCoreManualTimeSource.cpp
+++ b/Engine/Plugins/Experimental/PropertyAnimatorCore/Source/PropertyAnimatorCore/Private/TimeSources/PropertyAnimatorCoreManualTimeSource.cpp
@@ -5,20 +5,9 @@
 #include "Misc/App.h"
 #include "Presets/PropertyAnimatorCorePresetArchive.h"
 
-void UPropertyAnimatorCoreManualTimeSource::SetOverrideTime(bool bInOverride)
-{
-	if (bOverrideTime == bInOverride)
-	{
-		return;
-	}
-
-	bOverrideTime = bInOverride;
-	OnOverrideTimeChanged();
-}
-
 void UPropertyAnimatorCoreManualTimeSource::SetCustomTime(double InTime)
 {
-	if (!bOverrideTime || FMath::IsNearlyEqual(InTime, CustomTime))
+	if (FMath::IsNearlyEqual(InTime, CustomTime))
 	{
 		return;
 	}
@@ -26,49 +15,43 @@ void UPropertyAnimatorCoreManualTimeSource::SetCustomTime(double InTime)
 	CustomTime = InTime;
 }
 
-void UPropertyAnimatorCoreManualTimeSource::SetState(const FPropertyAnimatorCoreManualState& InState)
+void UPropertyAnimatorCoreManualTimeSource::SetPlaybackState(EPropertyAnimatorCoreManualStatus InState)
 {
-	if (State.Status == InState.Status)
+	if (PlaybackState == InState)
 	{
 		return;
 	}
 
-	State = InState;
+	PlaybackState = InState;
 	OnStateChanged();
 }
 
 bool UPropertyAnimatorCoreManualTimeSource::UpdateEvaluationData(FPropertyAnimatorCoreTimeSourceEvaluationData& OutData)
 {
-	if (!bOverrideTime)
+	/*
+	 * Don't use world delta time to avoid time dilation,
+	 * get the app to use raw time between frames and increment when this time source is enabled
+	 */
+	if (ActiveStatus == EPropertyAnimatorCoreManualStatus::PlayingForward)
+	{
+		CustomTime += FApp::GetDeltaTime();
+	}
+	else if (ActiveStatus == EPropertyAnimatorCoreManualStatus::PlayingBackward)
 	{
-		/*
-		 * Don't use world delta time to avoid time dilation,
-		 * get the app to use raw time between frames and increment when this time source is enabled
-		 */
-		if (ActiveStatus == EPropertyAnimatorCoreManualStatus::PlayingForward)
-		{
-			CustomTime += FApp::GetDeltaTime();
-		}
-		else if (ActiveStatus == EPropertyAnimatorCoreManualStatus::PlayingBackward)
-		{
-			CustomTime -= FApp::GetDeltaTime();
-		}
-
-		OutData.TimeElapsed = CustomTime;
-
-		return ActiveStatus != EPropertyAnimatorCoreManualStatus::Stopped;
+		CustomTime -= FApp::GetDeltaTime();
 	}
 
 	OutData.TimeElapsed = CustomTime;
 
-	return true;
+	return ActiveStatus != EPropertyAnimatorCoreManualStatus::Stopped;
 }
 
 void UPropertyAnimatorCoreManualTimeSource::OnTimeSourceActive()
 {
 	Super::OnTimeSourceActive();
 
-	ActiveStatus = EPropertyAnimatorCoreManualStatus::Stopped;
+	PlaybackState = EPropertyAnimatorCoreManualStatus::Paused;
+	ActiveStatus = PlaybackState;
 }
 
 void UPropertyAnimatorCoreManualTimeSource::OnTimeSourceInactive()
@@ -88,10 +71,6 @@ bool UPropertyAnimatorCoreManualTimeSource::ImportPreset(const UPropertyAnimator
 		ObjectArchive->Get(GET_MEMBER_NAME_STRING_CHECKED(UPropertyAnimatorCoreManualTimeSource, CustomTime), CustomTimeValue);
 		SetCustomTime(CustomTimeValue);
 
-		bool bOverrideTimeValue = bOverrideTime;
-		ObjectArchive->Get(GET_MEMBER_NAME_STRING_CHECKED(UPropertyAnimatorCoreManualTimeSource, bOverrideTime), bOverrideTimeValue);
-		SetOverrideTime(bOverrideTimeValue);
-
 		return true;
 	}
 
@@ -105,7 +84,6 @@ bool UPropertyAnimatorCoreManualTimeSource::ExportPreset(const UPropertyAnimator
 		TSharedPtr<FPropertyAnimatorCorePresetObjectArchive> ObjectArchive = OutValue->AsMutableObject();
 
 		ObjectArchive->Set(GET_MEMBER_NAME_STRING_CHECKED(UPropertyAnimatorCoreManualTimeSource, CustomTime), CustomTime);
-		ObjectArchive->Set(GET_MEMBER_NAME_STRING_CHECKED(UPropertyAnimatorCoreManualTimeSource, bOverrideTime), bOverrideTime);
 
 		return true;
 	}
@@ -115,11 +93,6 @@ bool UPropertyAnimatorCoreManualTimeSource::ExportPreset(const UPropertyAnimator
 
 void UPropertyAnimatorCoreManualTimeSource::Play(bool bInForward)
 {
-	if (bOverrideTime)
-	{
-		return;
-	}
-
 	// Allow change from playing forward to backward
 	if (!IsPlaying()
 		|| (bInForward && ActiveStatus == EPropertyAnimatorCoreManualStatus::PlayingBackward)
@@ -153,12 +126,12 @@ void UPropertyAnimatorCoreManualTimeSource::Stop()
 
 EPropertyAnimatorCoreManualStatus UPropertyAnimatorCoreManualTimeSource::GetPlaybackStatus() const
 {
-	return bOverrideTime ? EPropertyAnimatorCoreManualStatus::Stopped : ActiveStatus;
+	return ActiveStatus;
 }
 
 bool UPropertyAnimatorCoreManualTimeSource::IsPlaying() const
 {
-	return !bOverrideTime && (ActiveStatus == EPropertyAnimatorCoreManualStatus::PlayingBackward || ActiveStatus == EPropertyAnimatorCoreManualStatus::PlayingForward);
+	return ActiveStatus == EPropertyAnimatorCoreManualStatus::PlayingBackward || ActiveStatus == EPropertyAnimatorCoreManualStatus::PlayingForward;
 }
 
 #if WITH_EDITOR
@@ -168,20 +141,16 @@ void UPropertyAnimatorCoreManualTimeSource::PostEditChangeProperty(FPropertyChan
 
 	const FName MemberName = InPropertyChangedEvent.GetMemberPropertyName();
 
-	if (MemberName == GET_MEMBER_NAME_CHECKED(UPropertyAnimatorCoreManualTimeSource, State))
+	if (MemberName == GET_MEMBER_NAME_CHECKED(UPropertyAnimatorCoreManualTimeSource, PlaybackState))
 	{
 		OnStateChanged();
 	}
-	else if (MemberName == GET_MEMBER_NAME_CHECKED(UPropertyAnimatorCoreManualTimeSource, bOverrideTime))
-	{
-		OnOverrideTimeChanged();
-	}
 }
 #endif
 
 void UPropertyAnimatorCoreManualTimeSource::OnStateChanged()
 {
-	switch (State.Status)
+	switch (PlaybackState)
 	{
 	case EPropertyAnimatorCoreManualStatus::Stopped:
 		Stop();
@@ -197,10 +166,3 @@ void UPropertyAnimatorCoreManualTimeSource::OnStateChanged()
 		break;
 	}
 }
-
-void UPropertyAnimatorCoreManualTimeSource::OnOverrideTimeChanged()
-{
-	Stop();
-	CustomTime = 0.f;
-	State.Status = EPropertyAnimatorCoreManualStatus::Stopped;
-}
diff --git a/Engine/Plugins/Experimental/PropertyAnimatorCore/Source/PropertyAnimatorCore/Public/TimeSources/PropertyAnimatorCoreManualTimeSource.h b/Engine/Plugins/Experimental/PropertyAnimatorCore/Source/PropertyAnimatorCore/Public/TimeSources/PropertyAnimatorCoreManualTimeSource.h
--- a/Engine/Plugins/Experimental/PropertyAnimatorCore/Source/PropertyAnimatorCore/Public/TimeSources/PropertyAnimatorCoreManualTimeSource.h
+++ b/Engine/Plugins/Experimental/PropertyAnimatorCore/Source/PropertyAnimatorCore/Public/TimeSources/PropertyAnimatorCoreManualTimeSource.h
@@ -18,16 +18,6 @@ enum class EPropertyAnimatorCoreManualStatus : uint8
 	PlayingBackward = 1 << 3
 };
 
-/** Used to store status of player, for type customization */
-USTRUCT()
-struct FPropertyAnimatorCoreManualState
-{
-	GENERATED_BODY()
-
-	UPROPERTY(EditInstanceOnly, Category="Animator")
-	EPropertyAnimatorCoreManualStatus Status = EPropertyAnimatorCoreManualStatus::Stopped;
-};
-
 UCLASS(MinimalAPI)
 class UPropertyAnimatorCoreManualTimeSource : public UPropertyAnimatorCoreTimeSourceBase
 {
@@ -38,22 +28,16 @@ public:
 		: UPropertyAnimatorCoreTimeSourceBase(TEXT("Manual"))
 	{}
 
-	void SetOverrideTime(bool bInOverride);
-	bool GetOverrideTime() const
-	{
-		return bOverrideTime;
-	}
-
 	PROPERTYANIMATORCORE_API void SetCustomTime(double InTime);
 	double GetCustomTime() const
 	{
 		return CustomTime;
 	}
 
-	void SetState(const FPropertyAnimatorCoreManualState& InState);
-	const FPropertyAnimatorCoreManualState& GetState() const
+	PROPERTYANIMATORCORE_API void SetPlaybackState(EPropertyAnimatorCoreManualStatus InState);
+	EPropertyAnimatorCoreManualStatus GetPlaybackState() const
 	{
-		return State;
+		return PlaybackState;
 	}
 
 	void Play(bool bInForward);
@@ -79,20 +63,15 @@ protected:
 	//~ End UPropertyAnimatorTimeSourceBase
 
 	void OnStateChanged();
-	void OnOverrideTimeChanged();
-
-	/** Allows you to drive controllers with this float */
-	UPROPERTY(EditInstanceOnly, Setter="SetOverrideTime", Getter="GetOverrideTime", Category="Animator", meta=(InlineEditConditionToggle))
-	bool bOverrideTime = true;
 
 	/** Time to evaluate */
-	UPROPERTY(EditInstanceOnly, Setter, Getter, Category="Animator", meta=(EditCondition="bOverrideTime", Units=Seconds))
+	UPROPERTY(EditInstanceOnly, Setter, Getter, Category="Animator")
 	double CustomTime = 0.f;
 
 	/** Playback state */
-	UPROPERTY(EditInstanceOnly, Transient, DuplicateTransient, Setter, Getter, Category="Animator", meta=(EditCondition="!bOverrideTime", EditConditionHides, HideEditConditionToggles))
-	FPropertyAnimatorCoreManualState State;
+	UPROPERTY(EditInstanceOnly, Transient, DuplicateTransient, Setter, Getter, Category="Animator")
+	EPropertyAnimatorCoreManualStatus PlaybackState = EPropertyAnimatorCoreManualStatus::Paused;
 
 	/** Current active status for the player */
-	EPropertyAnimatorCoreManualStatus ActiveStatus = EPropertyAnimatorCoreManualStatus::Stopped;
+	EPropertyAnimatorCoreManualStatus ActiveStatus = EPropertyAnimatorCoreManualStatus::Paused;
 };
\ No newline at end of file
diff --git a/Engine/Plugins/Experimental/PropertyAnimatorCore/Source/PropertyAnimatorCoreEditor/Private/Customizations/PropertyAnimatorCoreEditorManualStateTypeCustomization.cpp b/Engine/Plugins/Experimental/PropertyAnimatorCore/Source/PropertyAnimatorCoreEditor/Private/Customizations/PropertyAnimatorCoreEditorManualStateTypeCustomization.cpp
--- a/Engine/Plugins/Experimental/PropertyAnimatorCore/Source/PropertyAnimatorCoreEditor/Private/Customizations/PropertyAnimatorCoreEditorManualStateTypeCustomization.cpp
+++ b/Engine/Plugins/Experimental/PropertyAnimatorCore/Source/PropertyAnimatorCoreEditor/Private/Customizations/PropertyAnimatorCoreEditorManualStateTypeCustomization.cpp
@@ -17,7 +17,7 @@ void FPropertyAnimatorCoreEditorManualStateTypeCustomization::CustomizeHeader(TS
 		return;
 	}
 
-	StatusPropertyHandle = InPropertyHandle->GetChildHandle(GET_MEMBER_NAME_CHECKED(FPropertyAnimatorCoreManualState, Status));
+	StatusPropertyHandle = InPropertyHandle;
 
 	InRow.NameContent()
 	[
@@ -48,7 +48,7 @@ void FPropertyAnimatorCoreEditorManualStateTypeCustomization::CustomizeHeader(TS
 		]
 		+ SHorizontalBox::Slot()
 		.AutoWidth()
-		.Padding(0.f, 0.f, 0.f, 5.f)
+		.Padding(5.f, 0.f, 0.f, 5.f)
 		.HAlign(HAlign_Fill)
 		[
 			SNew(SButton)
@@ -64,7 +64,7 @@ void FPropertyAnimatorCoreEditorManualStateTypeCustomization::CustomizeHeader(TS
 		]
 		+ SHorizontalBox::Slot()
 		.AutoWidth()
-		.Padding(0.f, 0.f, 0.f, 5.f)
+		.Padding(5.f, 0.f, 0.f, 5.f)
 		.HAlign(HAlign_Fill)
 		[
 			SNew(SButton)
@@ -80,7 +80,7 @@ void FPropertyAnimatorCoreEditorManualStateTypeCustomization::CustomizeHeader(TS
 		]
 		+ SHorizontalBox::Slot()
 		.AutoWidth()
-		.Padding(0.f, 0.f, 0.f, 5.f)
+		.Padding(5.f, 0.f, 0.f, 5.f)
 		.HAlign(HAlign_Fill)
 		[
 			SNew(SButton)
diff --git a/Engine/Plugins/Experimental/PropertyAnimatorCore/Source/PropertyAnimatorCoreEditor/Private/Customizations/PropertyAnimatorCoreEditorManualStateTypeCustomization.h b/Engine/Plugins/Experimental/PropertyAnimatorCore/Source/PropertyAnimatorCoreEditor/Private/Customizations/PropertyAnimatorCoreEditorManualStateTypeCustomization.h
--- a/Engine/Plugins/Experimental/PropertyAnimatorCore/Source/PropertyAnimatorCoreEditor/Private/Customizations/PropertyAnimatorCoreEditorManualStateTypeCustomization.h
+++ b/Engine/Plugins/Experimental/PropertyAnimatorCore/Source/PropertyAnimatorCoreEditor/Private/Customizations/PropertyAnimatorCoreEditorManualStateTypeCustomization.h
@@ -10,7 +10,7 @@ class IPropertyHandle;
 enum class EPropertyAnimatorCoreManualStatus : uint8;
 struct FSlateColor;
 
-/** Type customization for FPropertyAnimatorCoreManualState to show a player */
+/** Type customization for EPropertyAnimatorCoreManualStatus to show a player */
 class FPropertyAnimatorCoreEditorManualStateTypeCustomization : public IPropertyTypeCustomization
 {
 public:
diff --git a/Engine/Plugins/Experimental/PropertyAnimatorCore/Source/PropertyAnimatorCoreEditor/Private/PropertyAnimatorCoreEditorModule.cpp b/Engine/Plugins/Experimental/PropertyAnimatorCore/Source/PropertyAnimatorCoreEditor/Private/PropertyAnimatorCoreEditorModule.cpp
--- a/Engine/Plugins/Experimental/PropertyAnimatorCore/Source/PropertyAnimatorCoreEditor/Private/PropertyAnimatorCoreEditorModule.cpp
+++ b/Engine/Plugins/Experimental/PropertyAnimatorCore/Source/PropertyAnimatorCoreEditor/Private/PropertyAnimatorCoreEditorModule.cpp
@@ -20,12 +20,12 @@ void FPropertyAnimatorCoreEditorModule::StartupModule()
 {
 	FPropertyEditorModule& PropertyModule = FModuleManager::LoadModuleChecked<FPropertyEditorModule>("PropertyEditor");
 
-	PropertyModule.RegisterCustomPropertyTypeLayout(UPropertyAnimatorCoreContext::StaticClass()->GetFName(), FOnGetPropertyTypeCustomizationInstance::CreateStatic(&FPropertyAnimatorCoreEditorContextTypeCustomization::MakeInstance));
-	PropertyModule.RegisterCustomPropertyTypeLayout(FPropertyAnimatorCoreManualState::StaticStruct()->GetFName(), FOnGetPropertyTypeCustomizationInstance::CreateStatic(&FPropertyAnimatorCoreEditorManualStateTypeCustomization::MakeInstance));
-	PropertyModule.RegisterCustomPropertyTypeLayout(FPropertyAnimatorCoreSequencerTimeSourceEvalResult::StaticStruct()->GetFName(), FOnGetPropertyTypeCustomizationInstance::CreateStatic(&FPropertyAnimatorCoreEditorSequencerTimeSourceChannelTypeCustomization::MakeInstance));
-	PropertyModule.RegisterCustomPropertyTypeLayout(FIntProperty::StaticClass()->GetFName(), FOnGetPropertyTypeCustomizationInstance::CreateStatic(&FPropertyAnimatorCoreEditorSeedTypeCustomization::MakeInstance), MakeShared<FPropertyAnimatorCoreEditorSeedTypeIdentifier>());
+	PropertyModule.RegisterCustomPropertyTypeLayout(RegisteredCustomizations.Add_GetRef(UPropertyAnimatorCoreContext::StaticClass()->GetFName()), FOnGetPropertyTypeCustomizationInstance::CreateStatic(&FPropertyAnimatorCoreEditorContextTypeCustomization::MakeInstance));
+	PropertyModule.RegisterCustomPropertyTypeLayout(RegisteredCustomizations.Add_GetRef(StaticEnum<EPropertyAnimatorCoreManualStatus>()->GetFName()), FOnGetPropertyTypeCustomizationInstance::CreateStatic(&FPropertyAnimatorCoreEditorManualStateTypeCustomization::MakeInstance));
+	PropertyModule.RegisterCustomPropertyTypeLayout(RegisteredCustomizations.Add_GetRef(FPropertyAnimatorCoreSequencerTimeSourceEvalResult::StaticStruct()->GetFName()), FOnGetPropertyTypeCustomizationInstance::CreateStatic(&FPropertyAnimatorCoreEditorSequencerTimeSourceChannelTypeCustomization::MakeInstance));
+	PropertyModule.RegisterCustomPropertyTypeLayout(RegisteredCustomizations.Add_GetRef(FIntProperty::StaticClass()->GetFName()), FOnGetPropertyTypeCustomizationInstance::CreateStatic(&FPropertyAnimatorCoreEditorSeedTypeCustomization::MakeInstance), MakeShared<FPropertyAnimatorCoreEditorSeedTypeIdentifier>());
 
-	PropertyModule.RegisterCustomClassLayout(UPropertyAnimatorCoreBase::StaticClass()->GetFName(), FOnGetDetailCustomizationInstance::CreateStatic(&FPropertyAnimatorCoreEditorDetailCustomization::MakeInstance));
+	PropertyModule.RegisterCustomClassLayout(RegisteredCustomizations.Add_GetRef(UPropertyAnimatorCoreBase::StaticClass()->GetFName()), FOnGetDetailCustomizationInstance::CreateStatic(&FPropertyAnimatorCoreEditorDetailCustomization::MakeInstance));
 
 	ISequencerModule& SequencerModule = FModuleManager::LoadModuleChecked<ISequencerModule>("Sequencer");
 	AnimatorTrackCreateEditorHandle = SequencerModule.RegisterTrackEditor(FOnCreateTrackEditor::CreateStatic(&FMovieSceneAnimatorTrackEditor::CreateTrackEditor));
@@ -35,14 +35,15 @@ void FPropertyAnimatorCoreEditorModule::ShutdownModule()
 {
 	if (FPropertyEditorModule* PropertyModule = FModuleManager::GetModulePtr<FPropertyEditorModule>("PropertyEditor"))
 	{
-		PropertyModule->UnregisterCustomPropertyTypeLayout(UPropertyAnimatorCoreContext::StaticClass()->GetFName());
-		PropertyModule->UnregisterCustomPropertyTypeLayout(FPropertyAnimatorCoreManualState::StaticStruct()->GetFName());
-		PropertyModule->UnregisterCustomPropertyTypeLayout(FPropertyAnimatorCoreSequencerTimeSourceEvalResult::StaticStruct()->GetFName());
-		PropertyModule->UnregisterCustomPropertyTypeLayout(FIntProperty::StaticClass()->GetFName());
-
-		PropertyModule->UnregisterCustomClassLayout(UPropertyAnimatorCoreBase::StaticClass()->GetFName());
+		for (const FName& RegisteredCustomization : RegisteredCustomizations)
+		{
+			PropertyModule->UnregisterCustomPropertyTypeLayout(RegisteredCustomization);
+			PropertyModule->UnregisterCustomClassLayout(RegisteredCustomization);
+		}
 	}
 
+	RegisteredCustomizations.Empty();
+
 	if (ISequencerModule* SequencerModule = FModuleManager::GetModulePtr<ISequencerModule>("Sequencer"))
 	{
 		SequencerModule->UnRegisterTrackEditor(AnimatorTrackCreateEditorHandle);
diff --git a/Engine/Plugins/Experimental/PropertyAnimatorCore/Source/PropertyAnimatorCoreEditor/Private/PropertyAnimatorCoreEditorModule.h b/Engine/Plugins/Experimental/PropertyAnimatorCore/Source/PropertyAnimatorCoreEditor/Private/PropertyAnimatorCoreEditorModule.h
--- a/Engine/Plugins/Experimental/PropertyAnimatorCore/Source/PropertyAnimatorCoreEditor/Private/PropertyAnimatorCoreEditorModule.h
+++ b/Engine/Plugins/Experimental/PropertyAnimatorCore/Source/PropertyAnimatorCoreEditor/Private/PropertyAnimatorCoreEditorModule.h
@@ -2,7 +2,9 @@
 
 #pragma once
 
+#include "Containers/Array.h"
 #include "Modules/ModuleManager.h"
+#include "UObject/NameTypes.h"
 
 class FPropertyAnimatorCoreEditorModule : public IModuleInterface
 {
@@ -13,5 +15,6 @@ public:
 	//~ End IModuleInterface
 
 private:
+	TArray<FName> RegisteredCustomizations;
 	FDelegateHandle AnimatorTrackCreateEditorHandle;
 };
\ No newline at end of file
