MotionDesign : PropertyAnimator
- Changing animator manual time source bitflag enum to regular enum for playback status

diff --git a/Engine/Plugins/Experimental/PropertyAnimatorCore/Source/PropertyAnimatorCore/Public/TimeSources/PropertyAnimatorCoreManualTimeSource.h b/Engine/Plugins/Experimental/PropertyAnimatorCore/Source/PropertyAnimatorCore/Public/TimeSources/PropertyAnimatorCoreManualTimeSource.h
--- a/Engine/Plugins/Experimental/PropertyAnimatorCore/Source/PropertyAnimatorCore/Public/TimeSources/PropertyAnimatorCoreManualTimeSource.h
+++ b/Engine/Plugins/Experimental/PropertyAnimatorCore/Source/PropertyAnimatorCore/Public/TimeSources/PropertyAnimatorCoreManualTimeSource.h
@@ -9,13 +9,13 @@ UENUM()
 enum class EPropertyAnimatorCoreManualStatus : uint8
 {
 	/** Animation is done */
-	Stopped = 1 << 0,
+	Stopped,
 	/** Animation is paused */
-	Paused = 1 << 1,
+	Paused,
 	/** Animation is playing */
-	PlayingForward = 1 << 2,
+	PlayingForward,
 	/** Animation is playing in reverse */
-	PlayingBackward = 1 << 3
+	PlayingBackward
 };
 
 UCLASS(MinimalAPI)
@@ -68,8 +68,14 @@ protected:
 	UPROPERTY(EditInstanceOnly, Setter, Getter, Category="Animator")
 	double CustomTime = 0.f;
 
-	/** Playback state */
-	UPROPERTY(EditInstanceOnly, Transient, DuplicateTransient, Setter, Getter, Category="Animator")
+	/**
+	 * Playback state
+	 * Stopped = 0
+	 * Paused = 1
+	 * PlayingForward = 2
+	 * PlayingReverse = 3
+	 */
+	UPROPERTY(EditInstanceOnly, Transient, DuplicateTransient, NoClear, Setter, Getter, Category="Animator")
 	EPropertyAnimatorCoreManualStatus PlaybackState = EPropertyAnimatorCoreManualStatus::Paused;
 
 	/** Current active status for the player */
