MotionDesign : PropertyAnimator
- Adjusting edit conditions for property animator context

diff --git a/Engine/Plugins/Experimental/PropertyAnimatorCore/Source/PropertyAnimatorCore/Private/Properties/PropertyAnimatorCoreContext.cpp b/Engine/Plugins/Experimental/PropertyAnimatorCore/Source/PropertyAnimatorCore/Private/Properties/PropertyAnimatorCoreContext.cpp
--- a/Engine/Plugins/Experimental/PropertyAnimatorCore/Source/PropertyAnimatorCore/Private/Properties/PropertyAnimatorCoreContext.cpp
+++ b/Engine/Plugins/Experimental/PropertyAnimatorCore/Source/PropertyAnimatorCore/Private/Properties/PropertyAnimatorCoreContext.cpp
@@ -122,9 +122,7 @@ void UPropertyAnimatorCoreContext::PostLoad()
 {
 	Super::PostLoad();
 
-	CheckEditMode();
-	CheckEditConverterRule();
-	CheckEditResolver();
+	CheckEditConditions();
 
 	if (UPropertyAnimatorCoreResolver* PropertyResolver = AnimatedProperty.GetPropertyResolver())
 	{
@@ -276,8 +274,6 @@ bool UPropertyAnimatorCoreContext::ExportPreset(const UPropertyAnimatorCorePrese
 
 void UPropertyAnimatorCoreContext::OnAnimatedPropertyLinked()
 {
-	bEditMagnitude = AnimatedProperty.IsResolved();
-
 	if (const UPropertyAnimatorCoreResolver* PropertyResolver = AnimatedProperty.GetPropertyResolver())
 	{
 		if (!PropertyResolver->GetClass()->HasAnyClassFlags(CLASS_Abstract | CLASS_Transient))
@@ -336,9 +332,7 @@ bool UPropertyAnimatorCoreContext::ResolvePropertyOwner(AActor* InNewOwner)
 void UPropertyAnimatorCoreContext::ConstructInternal(const FPropertyAnimatorCoreData& InProperty)
 {
 	AnimatedProperty = InProperty;
-	CheckEditMode();
-	CheckEditConverterRule();
-	CheckEditResolver();
+	CheckEditConditions();
 	SetMode(EPropertyAnimatorCoreMode::Additive);
 	OnAnimatedPropertyLinked();
 }
@@ -390,6 +384,27 @@ void* UPropertyAnimatorCoreContext::GetConverterRulePtr(const UScriptStruct* InS
 	return nullptr;
 }
 
+void UPropertyAnimatorCoreContext::CheckEditConditions()
+{
+	CheckEditMagnitude();
+	CheckEditTimeOffset();
+	CheckEditMode();
+	CheckEditConverterRule();
+	CheckEditResolver();
+}
+
+void UPropertyAnimatorCoreContext::CheckEditMagnitude()
+{
+	bEditMagnitude = AnimatedProperty.IsA<FNumericProperty>()
+		|| AnimatedProperty.HasA<FNumericProperty>()
+		|| AnimatedProperty.IsA<FBoolProperty>();
+}
+
+void UPropertyAnimatorCoreContext::CheckEditTimeOffset()
+{
+	bEditTimeOffset = bEditMagnitude;	
+}
+
 void UPropertyAnimatorCoreContext::Restore()
 {
 	if (OriginalPropertyValues.GetNumPropertiesInBag() == 0
diff --git a/Engine/Plugins/Experimental/PropertyAnimatorCore/Source/PropertyAnimatorCore/Public/Properties/PropertyAnimatorCoreContext.h b/Engine/Plugins/Experimental/PropertyAnimatorCore/Source/PropertyAnimatorCore/Public/Properties/PropertyAnimatorCoreContext.h
--- a/Engine/Plugins/Experimental/PropertyAnimatorCore/Source/PropertyAnimatorCore/Public/Properties/PropertyAnimatorCoreContext.h
+++ b/Engine/Plugins/Experimental/PropertyAnimatorCore/Source/PropertyAnimatorCore/Public/Properties/PropertyAnimatorCoreContext.h
@@ -124,6 +124,9 @@ private:
 
 	PROPERTYANIMATORCORE_API void* GetConverterRulePtr(const UScriptStruct* InStruct);
 
+	void CheckEditConditions();
+	void CheckEditMagnitude();
+	void CheckEditTimeOffset();
 	void CheckEditMode();
 	void CheckEditConverterRule();
 	void CheckEditResolver();
@@ -154,34 +157,38 @@ private:
 	UPROPERTY(Transient)
 	bool bEditMagnitude = true;
 
+	/** Edit condition for TimeOffset */
+	UPROPERTY(Transient)
+	bool bEditTimeOffset = true;
+
+	/** Edit condition for modes */
+	UPROPERTY(Transient)
+	bool bEditMode = true;
+
+	/** Edit condition for converter rule */
+	UPROPERTY(Transient)
+	bool bEditConverterRule = false;
+
+	/** Edit condition for property resolver */
+	UPROPERTY(Transient)
+	bool bEditResolver = true;
+
 	/** Magnitude of the effect on this property */
 	UPROPERTY(EditInstanceOnly, Setter, Getter, Category="Animator", meta=(ClampMin="0.0", ClampMax="1.0", HideEditConditionToggle, EditCondition="bEditMagnitude", EditConditionHides, AllowPrivateAccess="true"))
 	float Magnitude = 1.f;
 
 	/** Time offset to evaluate this property */
-	UPROPERTY(EditInstanceOnly, Setter, Getter, Category="Animator", meta=(Units=Seconds, HideEditConditionToggle, EditCondition="bEditMagnitude", EditConditionHides, AllowPrivateAccess="true"))
+	UPROPERTY(EditInstanceOnly, Setter, Getter, Category="Animator", meta=(Units=Seconds, HideEditConditionToggle, EditCondition="bEditTimeOffset", EditConditionHides, AllowPrivateAccess="true"))
 	double TimeOffset = 0;
 
-	/** Edit condition for modes */
-	UPROPERTY(Transient)
-	bool bEditMode = true;
-
 	/** Current mode used for this property */
 	UPROPERTY(EditInstanceOnly, Setter, Getter, Category="Animator", meta=(HideEditConditionToggle, EditCondition="bEditMode", EditConditionHides, AllowPrivateAccess="true"))
 	EPropertyAnimatorCoreMode Mode = EPropertyAnimatorCoreMode::Absolute;
 
-	/** Edit condition for converter rule */
-	UPROPERTY(Transient)
-	bool bEditConverterRule = false;
-
 	/** If a converter is used, rules may be used to convert the property */
 	UPROPERTY(EditInstanceOnly, Category="Animator", NoClear, meta=(HideEditConditionToggle, EditCondition="bEditConverterRule", EditConditionHides, AllowPrivateAccess="true"))
 	TInstancedStruct<FPropertyAnimatorCoreConverterRuleBase> ConverterRule;
 
-	/** Edit condition for property resolver */
-	UPROPERTY(Transient)
-	bool bEditResolver = true;
-
 	/** Custom resolver for the property */
 	UPROPERTY(VisibleInstanceOnly, NoClear, Export, Instanced, DisplayName="Range", Category="Animator", meta=(EditCondition="bEditResolver", EditConditionHides, HideEditConditionToggle, AllowPrivateAccess="true"))
 	TObjectPtr<UPropertyAnimatorCoreResolver> Resolver;
