MotionDesign : Animators
- Set localization for animators in operator stack/menus

diff --git a/Engine/Plugins/VirtualProduction/Avalanche/Source/AvalanchePropertyAnimatorEditor/Private/Outliner/AvaPropertyAnimatorEditorOutliner.cpp b/Engine/Plugins/VirtualProduction/Avalanche/Source/AvalanchePropertyAnimatorEditor/Private/Outliner/AvaPropertyAnimatorEditorOutliner.cpp
--- a/Engine/Plugins/VirtualProduction/Avalanche/Source/AvalanchePropertyAnimatorEditor/Private/Outliner/AvaPropertyAnimatorEditorOutliner.cpp
+++ b/Engine/Plugins/VirtualProduction/Avalanche/Source/AvalanchePropertyAnimatorEditor/Private/Outliner/AvaPropertyAnimatorEditorOutliner.cpp
@@ -13,9 +13,10 @@ FAvaPropertyAnimatorEditorOutliner::FAvaPropertyAnimatorEditorOutliner(IAvaOutli
 	: FAvaOutlinerObject(InOutliner, InAnimator)
 	, PropertyAnimator(InAnimator)
 {
+	const TSharedPtr<const FPropertyAnimatorCoreMetadata> AnimatorMetadata = PropertyAnimator->GetAnimatorMetadata();
 	ItemName = FText::FromName(PropertyAnimator->GetAnimatorDisplayName());
 	ItemIcon = FSlateIconFinder::FindIconForClass(UPropertyAnimatorCoreComponent::StaticClass());
-	ItemTooltip = FText::FromName(PropertyAnimator->GetAnimatorOriginalName());
+	ItemTooltip = AnimatorMetadata->DisplayName;
 
 	UPropertyAnimatorCoreBase::OnPropertyAnimatorRemoved().AddRaw(this, &FAvaPropertyAnimatorEditorOutliner::OnAnimatorRemoved);
 }
diff --git a/Engine/Plugins/VirtualProduction/PropertyAnimator/Source/PropertyAnimator/Private/Animators/PropertyAnimatorBounce.cpp b/Engine/Plugins/VirtualProduction/PropertyAnimator/Source/PropertyAnimator/Private/Animators/PropertyAnimatorBounce.cpp
--- a/Engine/Plugins/VirtualProduction/PropertyAnimator/Source/PropertyAnimator/Private/Animators/PropertyAnimatorBounce.cpp
+++ b/Engine/Plugins/VirtualProduction/PropertyAnimator/Source/PropertyAnimator/Private/Animators/PropertyAnimatorBounce.cpp
@@ -6,6 +6,8 @@
 #include "Properties/PropertyAnimatorFloatContext.h"
 #include "PropertyAnimatorShared.h"
 
+#define LOCTEXT_NAMESPACE "PropertyAnimatorBounce"
+
 void UPropertyAnimatorBounce::SetInvertEffect(bool bInvert)
 {
 	if (bInvertEffect == bInvert)
@@ -22,6 +24,7 @@ void UPropertyAnimatorBounce::OnAnimatorRegistered(FPropertyAnimatorCoreMetadata
 	Super::OnAnimatorRegistered(InMetadata);
 
 	InMetadata.Name = TEXT("Bounce");
+	InMetadata.DisplayName = LOCTEXT("AnimatorDisplayName", "Bounce");
 }
 
 bool UPropertyAnimatorBounce::EvaluateProperty(const FPropertyAnimatorCoreData& InPropertyData, UPropertyAnimatorCoreContext* InContext, FInstancedPropertyBag& InParameters, FInstancedPropertyBag& OutEvaluationResult) const
@@ -40,3 +43,5 @@ bool UPropertyAnimatorBounce::EvaluateProperty(const FPropertyAnimatorCoreData&
 
 	return InContext->EvaluateProperty(InPropertyData, InParameters, OutEvaluationResult);
 }
+
+#undef LOCTEXT_NAMESPACE
diff --git a/Engine/Plugins/VirtualProduction/PropertyAnimator/Source/PropertyAnimator/Private/Animators/PropertyAnimatorClock.cpp b/Engine/Plugins/VirtualProduction/PropertyAnimator/Source/PropertyAnimator/Private/Animators/PropertyAnimatorClock.cpp
--- a/Engine/Plugins/VirtualProduction/PropertyAnimator/Source/PropertyAnimator/Private/Animators/PropertyAnimatorClock.cpp
+++ b/Engine/Plugins/VirtualProduction/PropertyAnimator/Source/PropertyAnimator/Private/Animators/PropertyAnimatorClock.cpp
@@ -7,6 +7,8 @@
 #include "Properties/Handlers/PropertyAnimatorCoreHandlerBase.h"
 #include "Subsystems/PropertyAnimatorCoreSubsystem.h"
 
+#define LOCTEXT_NAMESPACE "PropertyAnimatorClock"
+
 namespace UE::Private::Animator
 {
 	static TMap<TCHAR, TFunction<FString(const FDateTime&)>> FormatFunctions;
@@ -90,6 +92,7 @@ void UPropertyAnimatorClock::OnAnimatorRegistered(FPropertyAnimatorCoreMetadata&
 	Super::OnAnimatorRegistered(InMetadata);
 
 	InMetadata.Name = TEXT("Clock");
+	InMetadata.DisplayName = LOCTEXT("AnimatorDisplayName", "Clock");
 }
 
 void UPropertyAnimatorClock::EvaluateProperties(FInstancedPropertyBag& InParameters)
@@ -145,3 +148,5 @@ bool UPropertyAnimatorClock::ExportPreset(const UPropertyAnimatorCorePresetBase*
 
 	return false;
 }
+
+#undef LOCTEXT_NAMESPACE
diff --git a/Engine/Plugins/VirtualProduction/PropertyAnimator/Source/PropertyAnimator/Private/Animators/PropertyAnimatorCounter.cpp b/Engine/Plugins/VirtualProduction/PropertyAnimator/Source/PropertyAnimator/Private/Animators/PropertyAnimatorCounter.cpp
--- a/Engine/Plugins/VirtualProduction/PropertyAnimator/Source/PropertyAnimator/Private/Animators/PropertyAnimatorCounter.cpp
+++ b/Engine/Plugins/VirtualProduction/PropertyAnimator/Source/PropertyAnimator/Private/Animators/PropertyAnimatorCounter.cpp
@@ -6,6 +6,8 @@
 #include "Settings/PropertyAnimatorSettings.h"
 #include "Subsystems/PropertyAnimatorCoreSubsystem.h"
 
+#define LOCTEXT_NAMESPACE "PropertyAnimatorCounter"
+
 void FPropertyAnimatorCounterFormat::EnsureCharactersLength()
 {
 	if (DecimalCharacter.Len() > 1)
@@ -259,6 +261,7 @@ void UPropertyAnimatorCounter::OnAnimatorRegistered(FPropertyAnimatorCoreMetadat
 	Super::OnAnimatorRegistered(InMetadata);
 
 	InMetadata.Name = TEXT("Counter");
+	InMetadata.DisplayName = LOCTEXT("AnimatorDisplayName", "Counter");
 }
 
 void UPropertyAnimatorCounter::EvaluateProperties(FInstancedPropertyBag& InParameters)
@@ -393,3 +396,5 @@ void UPropertyAnimatorCounter::OnUseCustomFormatChanged()
 		}
 	}
 }
+
+#undef LOCTEXT_NAMESPACE
diff --git a/Engine/Plugins/VirtualProduction/PropertyAnimator/Source/PropertyAnimator/Private/Animators/PropertyAnimatorCurve.cpp b/Engine/Plugins/VirtualProduction/PropertyAnimator/Source/PropertyAnimator/Private/Animators/PropertyAnimatorCurve.cpp
--- a/Engine/Plugins/VirtualProduction/PropertyAnimator/Source/PropertyAnimator/Private/Animators/PropertyAnimatorCurve.cpp
+++ b/Engine/Plugins/VirtualProduction/PropertyAnimator/Source/PropertyAnimator/Private/Animators/PropertyAnimatorCurve.cpp
@@ -7,6 +7,8 @@
 #include "Properties/PropertyAnimatorFloatContext.h"
 #include "UObject/ConstructorHelpers.h"
 
+#define LOCTEXT_NAMESPACE "PropertyAnimatorCurve"
+
 UPropertyAnimatorCurve::UPropertyAnimatorCurve()
 {
 	static const ConstructorHelpers::FObjectFinder<UPropertyAnimatorWaveCurve> BaseCurve(TEXT("/Script/PropertyAnimator.PropertyAnimatorWaveCurve'/PropertyAnimator/Waves/Constant.Constant'"));
@@ -73,6 +75,7 @@ void UPropertyAnimatorCurve::OnAnimatorRegistered(FPropertyAnimatorCoreMetadata&
 	Super::OnAnimatorRegistered(InMetadata);
 
 	InMetadata.Name = TEXT("Curve");
+	InMetadata.DisplayName = LOCTEXT("AnimatorDisplayName", "Curve");
 }
 
 bool UPropertyAnimatorCurve::EvaluateProperty(const FPropertyAnimatorCoreData& InPropertyData, UPropertyAnimatorCoreContext* InContext, FInstancedPropertyBag& InParameters, FInstancedPropertyBag& OutEvaluationResult) const
@@ -243,3 +246,5 @@ void UPropertyAnimatorCurve::OnCycleDurationChanged()
 	OnEaseInChanged();
 	OnEaseOutChanged();
 }
+
+#undef LOCTEXT_NAMESPACE
diff --git a/Engine/Plugins/VirtualProduction/PropertyAnimator/Source/PropertyAnimator/Private/Animators/PropertyAnimatorOscillate.cpp b/Engine/Plugins/VirtualProduction/PropertyAnimator/Source/PropertyAnimator/Private/Animators/PropertyAnimatorOscillate.cpp
--- a/Engine/Plugins/VirtualProduction/PropertyAnimator/Source/PropertyAnimator/Private/Animators/PropertyAnimatorOscillate.cpp
+++ b/Engine/Plugins/VirtualProduction/PropertyAnimator/Source/PropertyAnimator/Private/Animators/PropertyAnimatorOscillate.cpp
@@ -5,6 +5,8 @@
 #include "Properties/PropertyAnimatorFloatContext.h"
 #include "PropertyAnimatorShared.h"
 
+#define LOCTEXT_NAMESPACE "PropertyAnimatorOscillate"
+
 void UPropertyAnimatorOscillate::SetOscillateFunction(EPropertyAnimatorOscillateFunction InFunction)
 {
 	OscillateFunction = InFunction;
@@ -15,6 +17,7 @@ void UPropertyAnimatorOscillate::OnAnimatorRegistered(FPropertyAnimatorCoreMetad
 	Super::OnAnimatorRegistered(InMetadata);
 
 	InMetadata.Name = TEXT("Oscillate");
+	InMetadata.DisplayName = LOCTEXT("AnimatorDisplayName", "Oscillate");
 }
 
 bool UPropertyAnimatorOscillate::EvaluateProperty(const FPropertyAnimatorCoreData& InPropertyData, UPropertyAnimatorCoreContext* InContext, FInstancedPropertyBag& InParameters, FInstancedPropertyBag& OutEvaluationResult) const
@@ -60,3 +63,5 @@ bool UPropertyAnimatorOscillate::EvaluateProperty(const FPropertyAnimatorCoreDat
 
 	return InContext->EvaluateProperty(InPropertyData, InParameters, OutEvaluationResult);
 }
+
+#undef LOCTEXT_NAMESPACE
diff --git a/Engine/Plugins/VirtualProduction/PropertyAnimator/Source/PropertyAnimator/Private/Animators/PropertyAnimatorPulse.cpp b/Engine/Plugins/VirtualProduction/PropertyAnimator/Source/PropertyAnimator/Private/Animators/PropertyAnimatorPulse.cpp
--- a/Engine/Plugins/VirtualProduction/PropertyAnimator/Source/PropertyAnimator/Private/Animators/PropertyAnimatorPulse.cpp
+++ b/Engine/Plugins/VirtualProduction/PropertyAnimator/Source/PropertyAnimator/Private/Animators/PropertyAnimatorPulse.cpp
@@ -5,6 +5,8 @@
 #include "Properties/PropertyAnimatorFloatContext.h"
 #include "PropertyAnimatorShared.h"
 
+#define LOCTEXT_NAMESPACE "PropertyAnimatorPulse"
+
 void UPropertyAnimatorPulse::SetEasingFunction(EPropertyAnimatorEasingFunction InEasingFunction)
 {
 	EasingFunction = InEasingFunction;
@@ -20,6 +22,7 @@ void UPropertyAnimatorPulse::OnAnimatorRegistered(FPropertyAnimatorCoreMetadata&
 	Super::OnAnimatorRegistered(InMetadata);
 
 	InMetadata.Name = TEXT("Pulse");
+	InMetadata.DisplayName = LOCTEXT("AnimatorDisplayName", "Pulse");
 }
 
 bool UPropertyAnimatorPulse::EvaluateProperty(const FPropertyAnimatorCoreData& InPropertyData, UPropertyAnimatorCoreContext* InContext, FInstancedPropertyBag& InParameters, FInstancedPropertyBag& OutEvaluationResult) const
@@ -43,3 +46,5 @@ bool UPropertyAnimatorPulse::EvaluateProperty(const FPropertyAnimatorCoreData& I
 
 	return InContext->EvaluateProperty(InPropertyData, InParameters, OutEvaluationResult);
 }
+
+#undef LOCTEXT_NAMESPACE
\ No newline at end of file
diff --git a/Engine/Plugins/VirtualProduction/PropertyAnimator/Source/PropertyAnimator/Private/Animators/PropertyAnimatorSoundWave.cpp b/Engine/Plugins/VirtualProduction/PropertyAnimator/Source/PropertyAnimator/Private/Animators/PropertyAnimatorSoundWave.cpp
--- a/Engine/Plugins/VirtualProduction/PropertyAnimator/Source/PropertyAnimator/Private/Animators/PropertyAnimatorSoundWave.cpp
+++ b/Engine/Plugins/VirtualProduction/PropertyAnimator/Source/PropertyAnimator/Private/Animators/PropertyAnimatorSoundWave.cpp
@@ -7,6 +7,8 @@
 #include "Properties/PropertyAnimatorFloatContext.h"
 #include "Sound/SoundWave.h"
 
+#define LOCTEXT_NAMESPACE "PropertyAnimatorSoundWave"
+
 UPropertyAnimatorSoundWave::UPropertyAnimatorSoundWave()
 {
 	CycleMode = EPropertyAnimatorCycleMode::None;
@@ -64,6 +66,7 @@ void UPropertyAnimatorSoundWave::OnAnimatorRegistered(FPropertyAnimatorCoreMetad
 	Super::OnAnimatorRegistered(InMetadata);
 
 	InMetadata.Name = TEXT("SoundWave");
+	InMetadata.DisplayName = LOCTEXT("AnimatorDisplayName", "SoundWave");
 }
 
 bool UPropertyAnimatorSoundWave::EvaluateProperty(const FPropertyAnimatorCoreData& InPropertyData, UPropertyAnimatorCoreContext* InContext, FInstancedPropertyBag& InParameters, FInstancedPropertyBag& OutEvaluationResult) const
@@ -136,3 +139,5 @@ bool UPropertyAnimatorSoundWave::ExportPreset(const UPropertyAnimatorCorePresetB
 
 	return false;
 }
+
+#undef LOCTEXT_NAMESPACE
\ No newline at end of file
diff --git a/Engine/Plugins/VirtualProduction/PropertyAnimator/Source/PropertyAnimator/Private/Animators/PropertyAnimatorTime.cpp b/Engine/Plugins/VirtualProduction/PropertyAnimator/Source/PropertyAnimator/Private/Animators/PropertyAnimatorTime.cpp
--- a/Engine/Plugins/VirtualProduction/PropertyAnimator/Source/PropertyAnimator/Private/Animators/PropertyAnimatorTime.cpp
+++ b/Engine/Plugins/VirtualProduction/PropertyAnimator/Source/PropertyAnimator/Private/Animators/PropertyAnimatorTime.cpp
@@ -5,11 +5,14 @@
 #include "Properties/Handlers/PropertyAnimatorCoreHandlerBase.h"
 #include "Properties/PropertyAnimatorFloatContext.h"
 
+#define LOCTEXT_NAMESPACE "PropertyAnimatorTime"
+
 void UPropertyAnimatorTime::OnAnimatorRegistered(FPropertyAnimatorCoreMetadata& InMetadata)
 {
 	Super::OnAnimatorRegistered(InMetadata);
 
 	InMetadata.Name = TEXT("Time");
+	InMetadata.DisplayName = LOCTEXT("AnimatorDisplayName", "Time");
 }
 
 bool UPropertyAnimatorTime::EvaluateProperty(const FPropertyAnimatorCoreData& InPropertyData, UPropertyAnimatorCoreContext* InContext, FInstancedPropertyBag& InParameters, FInstancedPropertyBag& OutEvaluationResult) const
@@ -25,3 +28,5 @@ bool UPropertyAnimatorTime::EvaluateProperty(const FPropertyAnimatorCoreData& In
 
 	return InContext->EvaluateProperty(InPropertyData, InParameters, OutEvaluationResult);
 }
+
+#undef LOCTEXT_NAMESPACE
diff --git a/Engine/Plugins/VirtualProduction/PropertyAnimator/Source/PropertyAnimator/Private/Animators/PropertyAnimatorWiggle.cpp b/Engine/Plugins/VirtualProduction/PropertyAnimator/Source/PropertyAnimator/Private/Animators/PropertyAnimatorWiggle.cpp
--- a/Engine/Plugins/VirtualProduction/PropertyAnimator/Source/PropertyAnimator/Private/Animators/PropertyAnimatorWiggle.cpp
+++ b/Engine/Plugins/VirtualProduction/PropertyAnimator/Source/PropertyAnimator/Private/Animators/PropertyAnimatorWiggle.cpp
@@ -6,6 +6,8 @@
 #include "Properties/PropertyAnimatorFloatContext.h"
 #include "PropertyAnimatorShared.h"
 
+#define LOCTEXT_NAMESPACE "PropertyAnimatorWiggle"
+
 UPropertyAnimatorWiggle::UPropertyAnimatorWiggle()
 {
 	static int32 SeedIncrement = 0;
@@ -25,6 +27,7 @@ void UPropertyAnimatorWiggle::OnAnimatorRegistered(FPropertyAnimatorCoreMetadata
 	Super::OnAnimatorRegistered(InMetadata);
 
 	InMetadata.Name = TEXT("Wiggle");
+	InMetadata.DisplayName = LOCTEXT("AnimatorDisplayName", "Wiggle");
 }
 
 bool UPropertyAnimatorWiggle::EvaluateProperty(const FPropertyAnimatorCoreData& InPropertyData, UPropertyAnimatorCoreContext* InContext, FInstancedPropertyBag& InParameters, FInstancedPropertyBag& OutEvaluationResult) const
@@ -72,3 +75,5 @@ bool UPropertyAnimatorWiggle::ExportPreset(const UPropertyAnimatorCorePresetBase
 
 	return false;
 }
+
+#undef LOCTEXT_NAMESPACE
diff --git a/Engine/Plugins/VirtualProduction/PropertyAnimator/Source/PropertyAnimator/Private/Presets/PropertyAnimatorPresetLocation.cpp b/Engine/Plugins/VirtualProduction/PropertyAnimator/Source/PropertyAnimator/Private/Presets/PropertyAnimatorPresetLocation.cpp
--- a/Engine/Plugins/VirtualProduction/PropertyAnimator/Source/PropertyAnimator/Private/Presets/PropertyAnimatorPresetLocation.cpp
+++ b/Engine/Plugins/VirtualProduction/PropertyAnimator/Source/PropertyAnimator/Private/Presets/PropertyAnimatorPresetLocation.cpp
@@ -6,6 +6,15 @@
 #include "Properties/PropertyAnimatorFloatContext.h"
 #include "Properties/PropertyAnimatorVectorContext.h"
 
+#define LOCTEXT_NAMESPACE "PropertyAnimatorPresetLocation"
+
+void UPropertyAnimatorPresetLocation::OnPresetRegistered()
+{
+	Super::OnPresetRegistered();
+
+	PresetDisplayName = LOCTEXT("PresetDisplayName", "Location");
+}
+
 void UPropertyAnimatorPresetLocation::GetPresetProperties(const AActor* InActor, const UPropertyAnimatorCoreBase* InAnimator, TSet<FPropertyAnimatorCoreData>& OutProperties) const
 {
 	if (!InActor->GetRootComponent())
@@ -40,3 +49,5 @@ void UPropertyAnimatorPresetLocation::OnPresetApplied(UPropertyAnimatorCoreBase*
 		}
 	}
 }
+
+#undef LOCTEXT_NAMESPACE
diff --git a/Engine/Plugins/VirtualProduction/PropertyAnimator/Source/PropertyAnimator/Private/Presets/PropertyAnimatorPresetLocation.h b/Engine/Plugins/VirtualProduction/PropertyAnimator/Source/PropertyAnimator/Private/Presets/PropertyAnimatorPresetLocation.h
--- a/Engine/Plugins/VirtualProduction/PropertyAnimator/Source/PropertyAnimator/Private/Presets/PropertyAnimatorPresetLocation.h
+++ b/Engine/Plugins/VirtualProduction/PropertyAnimator/Source/PropertyAnimator/Private/Presets/PropertyAnimatorPresetLocation.h
@@ -24,6 +24,7 @@ public:
 
 protected:
 	//~ Begin UPropertyAnimatorCorePresetBase
+	virtual void OnPresetRegistered() override;
 	virtual void GetPresetProperties(const AActor* InActor, const UPropertyAnimatorCoreBase* InAnimator, TSet<FPropertyAnimatorCoreData>& OutProperties) const override;
 	virtual void OnPresetApplied(UPropertyAnimatorCoreBase* InAnimator, const TSet<FPropertyAnimatorCoreData>& InProperties) override;
 	virtual bool LoadPreset() override { return true; }
diff --git a/Engine/Plugins/VirtualProduction/PropertyAnimator/Source/PropertyAnimator/Private/Presets/PropertyAnimatorPresetRotation.cpp b/Engine/Plugins/VirtualProduction/PropertyAnimator/Source/PropertyAnimator/Private/Presets/PropertyAnimatorPresetRotation.cpp
--- a/Engine/Plugins/VirtualProduction/PropertyAnimator/Source/PropertyAnimator/Private/Presets/PropertyAnimatorPresetRotation.cpp
+++ b/Engine/Plugins/VirtualProduction/PropertyAnimator/Source/PropertyAnimator/Private/Presets/PropertyAnimatorPresetRotation.cpp
@@ -6,6 +6,15 @@
 #include "Properties/PropertyAnimatorFloatContext.h"
 #include "Properties/PropertyAnimatorRotatorContext.h"
 
+#define LOCTEXT_NAMESPACE "PropertyAnimatorPresetRotation"
+
+void UPropertyAnimatorPresetRotation::OnPresetRegistered()
+{
+	Super::OnPresetRegistered();
+
+	PresetDisplayName = LOCTEXT("PresetDisplayName", "Rotation");
+}
+
 void UPropertyAnimatorPresetRotation::GetPresetProperties(const AActor* InActor, const UPropertyAnimatorCoreBase* InAnimator, TSet<FPropertyAnimatorCoreData>& OutProperties) const
 {
 	if (!InActor->GetRootComponent())
@@ -40,3 +49,5 @@ void UPropertyAnimatorPresetRotation::OnPresetApplied(UPropertyAnimatorCoreBase*
 		}
 	}
 }
+
+#undef LOCTEXT_NAMESPACE
diff --git a/Engine/Plugins/VirtualProduction/PropertyAnimator/Source/PropertyAnimator/Private/Presets/PropertyAnimatorPresetRotation.h b/Engine/Plugins/VirtualProduction/PropertyAnimator/Source/PropertyAnimator/Private/Presets/PropertyAnimatorPresetRotation.h
--- a/Engine/Plugins/VirtualProduction/PropertyAnimator/Source/PropertyAnimator/Private/Presets/PropertyAnimatorPresetRotation.h
+++ b/Engine/Plugins/VirtualProduction/PropertyAnimator/Source/PropertyAnimator/Private/Presets/PropertyAnimatorPresetRotation.h
@@ -24,6 +24,7 @@ public:
 
 protected:
 	//~ Begin UPropertyAnimatorCorePresetBase
+	virtual void OnPresetRegistered() override;
 	virtual void GetPresetProperties(const AActor* InActor, const UPropertyAnimatorCoreBase* InAnimator, TSet<FPropertyAnimatorCoreData>& OutProperties) const override;
 	virtual void OnPresetApplied(UPropertyAnimatorCoreBase* InAnimator, const TSet<FPropertyAnimatorCoreData>& InProperties) override;
 	virtual bool LoadPreset() override { return true; }
diff --git a/Engine/Plugins/VirtualProduction/PropertyAnimator/Source/PropertyAnimator/Private/Presets/PropertyAnimatorPresetScale.cpp b/Engine/Plugins/VirtualProduction/PropertyAnimator/Source/PropertyAnimator/Private/Presets/PropertyAnimatorPresetScale.cpp
--- a/Engine/Plugins/VirtualProduction/PropertyAnimator/Source/PropertyAnimator/Private/Presets/PropertyAnimatorPresetScale.cpp
+++ b/Engine/Plugins/VirtualProduction/PropertyAnimator/Source/PropertyAnimator/Private/Presets/PropertyAnimatorPresetScale.cpp
@@ -6,6 +6,15 @@
 #include "Properties/PropertyAnimatorFloatContext.h"
 #include "Properties/PropertyAnimatorVectorContext.h"
 
+#define LOCTEXT_NAMESPACE "PropertyAnimatorPresetScale"
+
+void UPropertyAnimatorPresetScale::OnPresetRegistered()
+{
+	Super::OnPresetRegistered();
+
+	PresetDisplayName = LOCTEXT("PresetDisplayName", "Scale");
+}
+
 void UPropertyAnimatorPresetScale::GetPresetProperties(const AActor* InActor, const UPropertyAnimatorCoreBase* InAnimator, TSet<FPropertyAnimatorCoreData>& OutProperties) const
 {
 	if (!InActor->GetRootComponent())
@@ -40,3 +49,5 @@ void UPropertyAnimatorPresetScale::OnPresetApplied(UPropertyAnimatorCoreBase* In
 		}
 	}
 }
+
+#undef LOCTEXT_NAMESPACE
diff --git a/Engine/Plugins/VirtualProduction/PropertyAnimator/Source/PropertyAnimator/Private/Presets/PropertyAnimatorPresetScale.h b/Engine/Plugins/VirtualProduction/PropertyAnimator/Source/PropertyAnimator/Private/Presets/PropertyAnimatorPresetScale.h
--- a/Engine/Plugins/VirtualProduction/PropertyAnimator/Source/PropertyAnimator/Private/Presets/PropertyAnimatorPresetScale.h
+++ b/Engine/Plugins/VirtualProduction/PropertyAnimator/Source/PropertyAnimator/Private/Presets/PropertyAnimatorPresetScale.h
@@ -24,6 +24,7 @@ public:
 
 protected:
 	//~ Begin UPropertyAnimatorCorePresetBase
+	virtual void OnPresetRegistered() override;
 	virtual void GetPresetProperties(const AActor* InActor, const UPropertyAnimatorCoreBase* InAnimator, TSet<FPropertyAnimatorCoreData>& OutProperties) const override;
 	virtual void OnPresetApplied(UPropertyAnimatorCoreBase* InAnimator, const TSet<FPropertyAnimatorCoreData>& InProperties) override;
 	virtual bool LoadPreset() override { return true; }
diff --git a/Engine/Plugins/VirtualProduction/PropertyAnimator/Source/PropertyAnimator/Private/Presets/PropertyAnimatorPresetText.cpp b/Engine/Plugins/VirtualProduction/PropertyAnimator/Source/PropertyAnimator/Private/Presets/PropertyAnimatorPresetText.cpp
--- a/Engine/Plugins/VirtualProduction/PropertyAnimator/Source/PropertyAnimator/Private/Presets/PropertyAnimatorPresetText.cpp
+++ b/Engine/Plugins/VirtualProduction/PropertyAnimator/Source/PropertyAnimator/Private/Presets/PropertyAnimatorPresetText.cpp
@@ -6,6 +6,15 @@
 #include "Properties/PropertyAnimatorCoreContext.h"
 #include "Text3DComponent.h"
 
+#define LOCTEXT_NAMESPACE "PropertyAnimatorPresetText"
+
+void UPropertyAnimatorPresetText::OnPresetRegistered()
+{
+	Super::OnPresetRegistered();
+
+	PresetDisplayName = LOCTEXT("PresetDisplayName", "Text");
+}
+
 void UPropertyAnimatorPresetText::GetPresetProperties(const AActor* InActor, const UPropertyAnimatorCoreBase* InAnimator, TSet<FPropertyAnimatorCoreData>& OutProperties) const
 {
 	const UText3DComponent* TextComponent = InActor->FindComponentByClass<UText3DComponent>();
@@ -33,3 +42,5 @@ void UPropertyAnimatorPresetText::OnPresetApplied(UPropertyAnimatorCoreBase* InA
 		}
 	}
 }
+
+#undef LOCTEXT_NAMESPACE
diff --git a/Engine/Plugins/VirtualProduction/PropertyAnimator/Source/PropertyAnimator/Private/Presets/PropertyAnimatorPresetText.h b/Engine/Plugins/VirtualProduction/PropertyAnimator/Source/PropertyAnimator/Private/Presets/PropertyAnimatorPresetText.h
--- a/Engine/Plugins/VirtualProduction/PropertyAnimator/Source/PropertyAnimator/Private/Presets/PropertyAnimatorPresetText.h
+++ b/Engine/Plugins/VirtualProduction/PropertyAnimator/Source/PropertyAnimator/Private/Presets/PropertyAnimatorPresetText.h
@@ -24,6 +24,7 @@ public:
 
 protected:
 	//~ Begin UPropertyAnimatorCorePresetBase
+	virtual void OnPresetRegistered() override;
 	virtual void GetPresetProperties(const AActor* InActor, const UPropertyAnimatorCoreBase* InAnimator, TSet<FPropertyAnimatorCoreData>& OutProperties) const override;
 	virtual void OnPresetApplied(UPropertyAnimatorCoreBase* InAnimator, const TSet<FPropertyAnimatorCoreData>& InProperties) override;
 	virtual bool LoadPreset() override { return true; }
diff --git a/Engine/Plugins/VirtualProduction/PropertyAnimator/Source/PropertyAnimator/Private/Presets/PropertyAnimatorPresetTextLocation.cpp b/Engine/Plugins/VirtualProduction/PropertyAnimator/Source/PropertyAnimator/Private/Presets/PropertyAnimatorPresetTextLocation.cpp
--- a/Engine/Plugins/VirtualProduction/PropertyAnimator/Source/PropertyAnimator/Private/Presets/PropertyAnimatorPresetTextLocation.cpp
+++ b/Engine/Plugins/VirtualProduction/PropertyAnimator/Source/PropertyAnimator/Private/Presets/PropertyAnimatorPresetTextLocation.cpp
@@ -8,6 +8,15 @@
 #include "Properties/PropertyAnimatorFloatContext.h"
 #include "Text3DComponent.h"
 
+#define LOCTEXT_NAMESPACE "PropertyAnimatorPresetTextLocation"
+
+void UPropertyAnimatorPresetTextLocation::OnPresetRegistered()
+{
+	Super::OnPresetRegistered();
+
+	PresetDisplayName = LOCTEXT("PresetDisplayName", "TextCharacterLocation");
+}
+
 void UPropertyAnimatorPresetTextLocation::GetPresetProperties(const AActor* InActor, const UPropertyAnimatorCoreBase* InAnimator, TSet<FPropertyAnimatorCoreData>& OutProperties) const
 {
 	UText3DComponent* TextComponent = InActor->FindComponentByClass<UText3DComponent>();
@@ -22,3 +31,5 @@ void UPropertyAnimatorPresetTextLocation::GetPresetProperties(const AActor* InAc
 
 	OutProperties.Add(FPropertyAnimatorCoreData(TextComponent, LocationProperty, nullptr, UPropertyAnimatorTextResolver::StaticClass()));
 }
+
+#undef LOCTEXT_NAMESPACE
diff --git a/Engine/Plugins/VirtualProduction/PropertyAnimator/Source/PropertyAnimator/Private/Presets/PropertyAnimatorPresetTextLocation.h b/Engine/Plugins/VirtualProduction/PropertyAnimator/Source/PropertyAnimator/Private/Presets/PropertyAnimatorPresetTextLocation.h
--- a/Engine/Plugins/VirtualProduction/PropertyAnimator/Source/PropertyAnimator/Private/Presets/PropertyAnimatorPresetTextLocation.h
+++ b/Engine/Plugins/VirtualProduction/PropertyAnimator/Source/PropertyAnimator/Private/Presets/PropertyAnimatorPresetTextLocation.h
@@ -24,6 +24,7 @@ public:
 
 protected:
 	//~ Begin UPropertyAnimatorCorePresetBase
+	virtual void OnPresetRegistered() override;
 	virtual void GetPresetProperties(const AActor* InActor, const UPropertyAnimatorCoreBase* InAnimator, TSet<FPropertyAnimatorCoreData>& OutProperties) const override;
 	//~ End UPropertyAnimatorCorePresetBase
 };
\ No newline at end of file
diff --git a/Engine/Plugins/VirtualProduction/PropertyAnimator/Source/PropertyAnimator/Private/Presets/PropertyAnimatorPresetTextRotation.cpp b/Engine/Plugins/VirtualProduction/PropertyAnimator/Source/PropertyAnimator/Private/Presets/PropertyAnimatorPresetTextRotation.cpp
--- a/Engine/Plugins/VirtualProduction/PropertyAnimator/Source/PropertyAnimator/Private/Presets/PropertyAnimatorPresetTextRotation.cpp
+++ b/Engine/Plugins/VirtualProduction/PropertyAnimator/Source/PropertyAnimator/Private/Presets/PropertyAnimatorPresetTextRotation.cpp
@@ -8,6 +8,15 @@
 #include "Properties/PropertyAnimatorFloatContext.h"
 #include "Text3DComponent.h"
 
+#define LOCTEXT_NAMESPACE "PropertyAnimatorPresetTextRotation"
+
+void UPropertyAnimatorPresetTextRotation::OnPresetRegistered()
+{
+	Super::OnPresetRegistered();
+
+	PresetDisplayName = LOCTEXT("PresetDisplayName", "TextCharacterRotation");
+}
+
 void UPropertyAnimatorPresetTextRotation::GetPresetProperties(const AActor* InActor, const UPropertyAnimatorCoreBase* InAnimator, TSet<FPropertyAnimatorCoreData>& OutProperties) const
 {
 	UText3DComponent* TextComponent = InActor->FindComponentByClass<UText3DComponent>();
@@ -22,3 +31,5 @@ void UPropertyAnimatorPresetTextRotation::GetPresetProperties(const AActor* InAc
 
 	OutProperties.Add(FPropertyAnimatorCoreData(TextComponent, RotationProperty, nullptr, UPropertyAnimatorTextResolver::StaticClass()));
 }
+
+#undef LOCTEXT_NAMESPACE
diff --git a/Engine/Plugins/VirtualProduction/PropertyAnimator/Source/PropertyAnimator/Private/Presets/PropertyAnimatorPresetTextRotation.h b/Engine/Plugins/VirtualProduction/PropertyAnimator/Source/PropertyAnimator/Private/Presets/PropertyAnimatorPresetTextRotation.h
--- a/Engine/Plugins/VirtualProduction/PropertyAnimator/Source/PropertyAnimator/Private/Presets/PropertyAnimatorPresetTextRotation.h
+++ b/Engine/Plugins/VirtualProduction/PropertyAnimator/Source/PropertyAnimator/Private/Presets/PropertyAnimatorPresetTextRotation.h
@@ -24,6 +24,7 @@ public:
 
 protected:
 	//~ Begin UPropertyAnimatorCorePresetBase
+	virtual void OnPresetRegistered() override;
 	virtual void GetPresetProperties(const AActor* InActor, const UPropertyAnimatorCoreBase* InAnimator, TSet<FPropertyAnimatorCoreData>& OutProperties) const override;
 	//~ End UPropertyAnimatorCorePresetBase
 };
\ No newline at end of file
diff --git a/Engine/Plugins/VirtualProduction/PropertyAnimator/Source/PropertyAnimator/Private/Presets/PropertyAnimatorPresetTextScale.cpp b/Engine/Plugins/VirtualProduction/PropertyAnimator/Source/PropertyAnimator/Private/Presets/PropertyAnimatorPresetTextScale.cpp
--- a/Engine/Plugins/VirtualProduction/PropertyAnimator/Source/PropertyAnimator/Private/Presets/PropertyAnimatorPresetTextScale.cpp
+++ b/Engine/Plugins/VirtualProduction/PropertyAnimator/Source/PropertyAnimator/Private/Presets/PropertyAnimatorPresetTextScale.cpp
@@ -8,6 +8,15 @@
 #include "Properties/PropertyAnimatorFloatContext.h"
 #include "Text3DComponent.h"
 
+#define LOCTEXT_NAMESPACE "PropertyAnimatorPresetTextScale"
+
+void UPropertyAnimatorPresetTextScale::OnPresetRegistered()
+{
+	Super::OnPresetRegistered();
+
+	PresetDisplayName = LOCTEXT("PresetDisplayName", "TextCharacterScale");
+}
+
 void UPropertyAnimatorPresetTextScale::GetPresetProperties(const AActor* InActor, const UPropertyAnimatorCoreBase* InAnimator, TSet<FPropertyAnimatorCoreData>& OutProperties) const
 {
 	UText3DComponent* TextComponent = InActor->FindComponentByClass<UText3DComponent>();
@@ -22,3 +31,5 @@ void UPropertyAnimatorPresetTextScale::GetPresetProperties(const AActor* InActor
 
 	OutProperties.Add(FPropertyAnimatorCoreData(TextComponent, ScaleProperty, nullptr, UPropertyAnimatorTextResolver::StaticClass()));
 }
+
+#undef LOCTEXT_NAMESPACE
diff --git a/Engine/Plugins/VirtualProduction/PropertyAnimator/Source/PropertyAnimator/Private/Presets/PropertyAnimatorPresetTextScale.h b/Engine/Plugins/VirtualProduction/PropertyAnimator/Source/PropertyAnimator/Private/Presets/PropertyAnimatorPresetTextScale.h
--- a/Engine/Plugins/VirtualProduction/PropertyAnimator/Source/PropertyAnimator/Private/Presets/PropertyAnimatorPresetTextScale.h
+++ b/Engine/Plugins/VirtualProduction/PropertyAnimator/Source/PropertyAnimator/Private/Presets/PropertyAnimatorPresetTextScale.h
@@ -24,6 +24,7 @@ public:
 
 protected:
 	//~ Begin UPropertyAnimatorCorePresetBase
+	virtual void OnPresetRegistered() override;
 	virtual void GetPresetProperties(const AActor* InActor, const UPropertyAnimatorCoreBase* InAnimator, TSet<FPropertyAnimatorCoreData>& OutProperties) const override;
 	//~ End UPropertyAnimatorCorePresetBase
 };
\ No newline at end of file
diff --git a/Engine/Plugins/VirtualProduction/PropertyAnimator/Source/PropertyAnimator/Private/Presets/PropertyAnimatorPresetTextVisibility.cpp b/Engine/Plugins/VirtualProduction/PropertyAnimator/Source/PropertyAnimator/Private/Presets/PropertyAnimatorPresetTextVisibility.cpp
--- a/Engine/Plugins/VirtualProduction/PropertyAnimator/Source/PropertyAnimator/Private/Presets/PropertyAnimatorPresetTextVisibility.cpp
+++ b/Engine/Plugins/VirtualProduction/PropertyAnimator/Source/PropertyAnimator/Private/Presets/PropertyAnimatorPresetTextVisibility.cpp
@@ -7,6 +7,15 @@
 #include "Properties/PropertyAnimatorTextResolver.h"
 #include "Text3DComponent.h"
 
+#define LOCTEXT_NAMESPACE "PropertyAnimatorPresetTextVisibility"
+
+void UPropertyAnimatorPresetTextVisibility::OnPresetRegistered()
+{
+	Super::OnPresetRegistered();
+
+	PresetDisplayName = LOCTEXT("PresetDisplayName", "TextCharacterVisibility");
+}
+
 void UPropertyAnimatorPresetTextVisibility::GetPresetProperties(const AActor* InActor, const UPropertyAnimatorCoreBase* InAnimator, TSet<FPropertyAnimatorCoreData>& OutProperties) const
 {
 	UText3DComponent* TextComponent = InActor->FindComponentByClass<UText3DComponent>();
@@ -21,3 +30,5 @@ void UPropertyAnimatorPresetTextVisibility::GetPresetProperties(const AActor* In
 
 	OutProperties.Add(FPropertyAnimatorCoreData(TextComponent, VisibilityProperty, nullptr, UPropertyAnimatorTextResolver::StaticClass()));
 }
+
+#undef LOCTEXT_NAMESPACE
diff --git a/Engine/Plugins/VirtualProduction/PropertyAnimator/Source/PropertyAnimator/Private/Presets/PropertyAnimatorPresetTextVisibility.h b/Engine/Plugins/VirtualProduction/PropertyAnimator/Source/PropertyAnimator/Private/Presets/PropertyAnimatorPresetTextVisibility.h
--- a/Engine/Plugins/VirtualProduction/PropertyAnimator/Source/PropertyAnimator/Private/Presets/PropertyAnimatorPresetTextVisibility.h
+++ b/Engine/Plugins/VirtualProduction/PropertyAnimator/Source/PropertyAnimator/Private/Presets/PropertyAnimatorPresetTextVisibility.h
@@ -24,6 +24,7 @@ public:
 
 protected:
 	//~ Begin UPropertyAnimatorCorePresetBase
+	virtual void OnPresetRegistered() override;
 	virtual void GetPresetProperties(const AActor* InActor, const UPropertyAnimatorCoreBase* InAnimator, TSet<FPropertyAnimatorCoreData>& OutProperties) const override;
 	//~ End UPropertyAnimatorCorePresetBase
 };
\ No newline at end of file
diff --git a/Engine/Plugins/VirtualProduction/PropertyAnimator/Source/PropertyAnimator/Private/Presets/PropertyAnimatorPresetVisibility.cpp b/Engine/Plugins/VirtualProduction/PropertyAnimator/Source/PropertyAnimator/Private/Presets/PropertyAnimatorPresetVisibility.cpp
--- a/Engine/Plugins/VirtualProduction/PropertyAnimator/Source/PropertyAnimator/Private/Presets/PropertyAnimatorPresetVisibility.cpp
+++ b/Engine/Plugins/VirtualProduction/PropertyAnimator/Source/PropertyAnimator/Private/Presets/PropertyAnimatorPresetVisibility.cpp
@@ -5,6 +5,15 @@
 #include "Animators/PropertyAnimatorCoreBase.h"
 #include "Properties/Converters/PropertyAnimatorCoreConverterTraits.h"
 
+#define LOCTEXT_NAMESPACE "PropertyAnimatorPresetVisibility"
+
+void UPropertyAnimatorPresetVisibility::OnPresetRegistered()
+{
+	Super::OnPresetRegistered();
+
+	PresetDisplayName = LOCTEXT("PresetDisplayName", "Visibility");
+}
+
 void UPropertyAnimatorPresetVisibility::GetPresetProperties(const AActor* InActor, const UPropertyAnimatorCoreBase* InAnimator, TSet<FPropertyAnimatorCoreData>& OutProperties) const
 {
 	if (!InActor->GetRootComponent())
@@ -37,3 +46,5 @@ void UPropertyAnimatorPresetVisibility::OnPresetApplied(UPropertyAnimatorCoreBas
 		}
 	}
 }
+
+#undef LOCTEXT_NAMESPACE
diff --git a/Engine/Plugins/VirtualProduction/PropertyAnimator/Source/PropertyAnimator/Private/Presets/PropertyAnimatorPresetVisibility.h b/Engine/Plugins/VirtualProduction/PropertyAnimator/Source/PropertyAnimator/Private/Presets/PropertyAnimatorPresetVisibility.h
--- a/Engine/Plugins/VirtualProduction/PropertyAnimator/Source/PropertyAnimator/Private/Presets/PropertyAnimatorPresetVisibility.h
+++ b/Engine/Plugins/VirtualProduction/PropertyAnimator/Source/PropertyAnimator/Private/Presets/PropertyAnimatorPresetVisibility.h
@@ -24,6 +24,7 @@ public:
 
 protected:
 	//~ Begin UPropertyAnimatorCorePresetBase
+	virtual void OnPresetRegistered() override;
 	virtual void GetPresetProperties(const AActor* InActor, const UPropertyAnimatorCoreBase* InAnimator, TSet<FPropertyAnimatorCoreData>& OutProperties) const override;
 	virtual void OnPresetApplied(UPropertyAnimatorCoreBase* InAnimator, const TSet<FPropertyAnimatorCoreData>& InProperties) override;
 	virtual bool LoadPreset() override { return true; }
diff --git a/Engine/Plugins/VirtualProduction/PropertyAnimatorCore/Source/PropertyAnimatorCore/Private/Animators/PropertyAnimatorCoreBase.cpp b/Engine/Plugins/VirtualProduction/PropertyAnimatorCore/Source/PropertyAnimatorCore/Private/Animators/PropertyAnimatorCoreBase.cpp
--- a/Engine/Plugins/VirtualProduction/PropertyAnimatorCore/Source/PropertyAnimatorCore/Private/Animators/PropertyAnimatorCoreBase.cpp
+++ b/Engine/Plugins/VirtualProduction/PropertyAnimatorCore/Source/PropertyAnimatorCore/Private/Animators/PropertyAnimatorCoreBase.cpp
@@ -96,7 +96,7 @@ void UPropertyAnimatorCoreBase::UpdateAnimatorDisplayName()
 	}
 	else
 	{
-		SetAnimatorDisplayName(FName(GetAnimatorOriginalName().ToString() + TEXT("_") + CommonPrefix));
+		SetAnimatorDisplayName(FName(GetAnimatorMetadata()->Name.ToString() + TEXT("_") + CommonPrefix));
 	}
 }
 
@@ -117,7 +117,20 @@ void UPropertyAnimatorCoreBase::PostCDOContruct()
 	if (IsTemplate() && !Metadata.IsValid())
 	{
 		Metadata = MakeShared<FPropertyAnimatorCoreMetadata>();
+
+#if WITH_EDITOR
+		const UClass* AnimatorClass = GetClass();
+		Metadata->DisplayName = AnimatorClass->GetDisplayNameText();
+		Metadata->Description = AnimatorClass->GetToolTipText();
+#endif
+
 		OnAnimatorRegistered(*Metadata);
+
+		if (Metadata->DisplayName.IsEmpty())
+		{
+			Metadata->DisplayName = FText::FromName(Metadata->Name);
+		}
+
 		SetAnimatorDisplayName(Metadata->Name);
 	}
 }
@@ -378,24 +391,10 @@ void UPropertyAnimatorCoreBase::SetTimeSourceName(FName InTimeSourceName)
 	OnTimeSourceNameChanged();
 }
 
-FName UPropertyAnimatorCoreBase::GetAnimatorCategory() const
+TSharedRef<const FPropertyAnimatorCoreMetadata> UPropertyAnimatorCoreBase::GetAnimatorMetadata() const
 {
-	if (Metadata.IsValid())
-	{
-		return Metadata->Category;
-	}
-
-	return NAME_None;
-}
-
-FName UPropertyAnimatorCoreBase::GetAnimatorOriginalName() const
-{
-	if (Metadata.IsValid())
-	{
-		return Metadata->Name;
-	}
-
-	return NAME_None;
+	checkf(Metadata.IsValid(), TEXT("Animator metadata is invalid"))
+	return Metadata.ToSharedRef();
 }
 
 bool UPropertyAnimatorCoreBase::GetPropertiesSupported(const FPropertyAnimatorCoreData& InPropertyData, TSet<FPropertyAnimatorCoreData>& OutProperties, uint8 InSearchDepth, EPropertyAnimatorPropertySupport InSupportExpected) const
@@ -503,7 +502,7 @@ void UPropertyAnimatorCoreBase::OnAnimatorEnabled(EPropertyAnimatorCoreUpdateEve
 		, TEXT("%s : PropertyAnimator %s (%s) enabled %i")
 		, GetAnimatorActor() ? *GetAnimatorActor()->GetActorNameOrLabel() : TEXT("Invalid Actor")
 		, *GetAnimatorDisplayName().ToString()
-		, *GetAnimatorOriginalName().ToString()
+		, *GetAnimatorMetadata()->Name.ToString()
 		, InType);
 }
 
@@ -514,7 +513,7 @@ void UPropertyAnimatorCoreBase::OnAnimatorDisabled(EPropertyAnimatorCoreUpdateEv
 		, TEXT("%s : PropertyAnimator %s (%s) disabled %i")
 		, GetAnimatorActor() ? *GetAnimatorActor()->GetActorNameOrLabel() : TEXT("Invalid Actor")
 		, *GetAnimatorDisplayName().ToString()
-		, *GetAnimatorOriginalName().ToString()
+		, *GetAnimatorMetadata()->Name.ToString()
 		, InType);
 
 	constexpr bool bForceReset = true;
diff --git a/Engine/Plugins/VirtualProduction/PropertyAnimatorCore/Source/PropertyAnimatorCore/Private/Components/PropertyAnimatorCoreComponent.cpp b/Engine/Plugins/VirtualProduction/PropertyAnimatorCore/Source/PropertyAnimatorCore/Private/Components/PropertyAnimatorCoreComponent.cpp
--- a/Engine/Plugins/VirtualProduction/PropertyAnimatorCore/Source/PropertyAnimatorCore/Private/Components/PropertyAnimatorCoreComponent.cpp
+++ b/Engine/Plugins/VirtualProduction/PropertyAnimatorCore/Source/PropertyAnimatorCore/Private/Components/PropertyAnimatorCoreComponent.cpp
@@ -219,7 +219,7 @@ FName UPropertyAnimatorCoreComponent::GetAnimatorName(const UPropertyAnimatorCor
 
 	FString NewAnimatorName = InAnimator->GetName();
 
-	const int32 Idx = NewAnimatorName.Find(InAnimator->GetAnimatorOriginalName().ToString());
+	const int32 Idx = NewAnimatorName.Find(InAnimator->GetAnimatorMetadata()->Name.ToString());
 	if (Idx != INDEX_NONE)
 	{
 		NewAnimatorName = NewAnimatorName.RightChop(Idx);
diff --git a/Engine/Plugins/VirtualProduction/PropertyAnimatorCore/Source/PropertyAnimatorCore/Private/Presets/PropertyAnimatorCorePresetBase.cpp b/Engine/Plugins/VirtualProduction/PropertyAnimatorCore/Source/PropertyAnimatorCore/Private/Presets/PropertyAnimatorCorePresetBase.cpp
--- a/Engine/Plugins/VirtualProduction/PropertyAnimatorCore/Source/PropertyAnimatorCore/Private/Presets/PropertyAnimatorCorePresetBase.cpp
+++ b/Engine/Plugins/VirtualProduction/PropertyAnimatorCore/Source/PropertyAnimatorCore/Private/Presets/PropertyAnimatorCorePresetBase.cpp
@@ -10,9 +10,9 @@
 #include "Serialization/JsonWriter.h"
 #include "Serialization/JsonSerializer.h"
 
-FString UPropertyAnimatorCorePresetBase::GetPresetDisplayName() const
+FText UPropertyAnimatorCorePresetBase::GetPresetDisplayName() const
 {
-	return FName::NameToDisplayString(PresetName.ToString(), false);
+	return !PresetDisplayName.IsEmpty() ? PresetDisplayName : FText::FromString(FName::NameToDisplayString(PresetName.ToString(), /** IsBool */false));
 }
 
 void UPropertyAnimatorCorePresetBase::CreatePreset(FName InName, const TArray<IPropertyAnimatorCorePresetable*>& InPresetableItem)
diff --git a/Engine/Plugins/VirtualProduction/PropertyAnimatorCore/Source/PropertyAnimatorCore/Private/Presets/PropertyAnimatorCorePropertyPreset.cpp b/Engine/Plugins/VirtualProduction/PropertyAnimatorCore/Source/PropertyAnimatorCore/Private/Presets/PropertyAnimatorCorePropertyPreset.cpp
--- a/Engine/Plugins/VirtualProduction/PropertyAnimatorCore/Source/PropertyAnimatorCore/Private/Presets/PropertyAnimatorCorePropertyPreset.cpp
+++ b/Engine/Plugins/VirtualProduction/PropertyAnimatorCore/Source/PropertyAnimatorCore/Private/Presets/PropertyAnimatorCorePropertyPreset.cpp
@@ -23,7 +23,7 @@ void UPropertyAnimatorCorePropertyPreset::OnPresetApplied(UPropertyAnimatorCoreB
 {
 	if (PropertyPresets.IsEmpty())
 	{
-		UE_LOG(LogPropertyAnimatorCore, Log, TEXT("Preset %s applied on %s animator with %i properties"), *GetPresetName().ToString(), *InAnimator->GetAnimatorOriginalName().ToString(), InProperties.Num())
+		UE_LOG(LogPropertyAnimatorCore, Log, TEXT("Preset %s applied on %s animator with %i properties"), *GetPresetName().ToString(), *InAnimator->GetAnimatorMetadata()->Name.ToString(), InProperties.Num())
 		return;
 	}
 	
@@ -57,11 +57,11 @@ void UPropertyAnimatorCorePropertyPreset::OnPresetApplied(UPropertyAnimatorCoreB
 			{
 				if (Context->ImportPreset(this, PropertyPreset.Value))
 				{
-					UE_LOG(LogPropertyAnimatorCore, Log, TEXT("Successfully imported preset %s on %s animator for property %s"), *GetPresetName().ToString(), *InAnimator->GetAnimatorOriginalName().ToString(), *Property.GetPropertyDisplayName())
+					UE_LOG(LogPropertyAnimatorCore, Log, TEXT("Successfully imported preset %s on %s animator for property %s"), *GetPresetName().ToString(), *InAnimator->GetAnimatorMetadata()->Name.ToString(), *Property.GetPropertyDisplayName())
 				}
 				else
 				{
-					UE_LOG(LogPropertyAnimatorCore, Warning, TEXT("Failed to import preset %s on %s animator for property %s"), *GetPresetName().ToString(), *InAnimator->GetAnimatorOriginalName().ToString(), *Property.GetPropertyDisplayName())
+					UE_LOG(LogPropertyAnimatorCore, Warning, TEXT("Failed to import preset %s on %s animator for property %s"), *GetPresetName().ToString(), *InAnimator->GetAnimatorMetadata()->Name.ToString(), *Property.GetPropertyDisplayName())
 				}
 			}
 		}
@@ -121,7 +121,7 @@ bool UPropertyAnimatorCorePropertyPreset::ApplyPreset(UPropertyAnimatorCoreBase*
 		}
 	}
 
-	InAnimator->SetAnimatorDisplayName(FName(InAnimator->GetAnimatorOriginalName().ToString() + TEXT("_") + GetPresetDisplayName()));
+	InAnimator->SetAnimatorDisplayName(FName(InAnimator->GetAnimatorMetadata()->DisplayName.ToString() + TEXT("_") + GetPresetDisplayName().ToString()));
 
 	OnPresetApplied(InAnimator, SupportedProperties);
 
diff --git a/Engine/Plugins/VirtualProduction/PropertyAnimatorCore/Source/PropertyAnimatorCore/Private/Subsystems/PropertyAnimatorCoreSubsystem.cpp b/Engine/Plugins/VirtualProduction/PropertyAnimatorCore/Source/PropertyAnimatorCore/Private/Subsystems/PropertyAnimatorCoreSubsystem.cpp
--- a/Engine/Plugins/VirtualProduction/PropertyAnimatorCore/Source/PropertyAnimatorCore/Private/Subsystems/PropertyAnimatorCoreSubsystem.cpp
+++ b/Engine/Plugins/VirtualProduction/PropertyAnimatorCore/Source/PropertyAnimatorCore/Private/Subsystems/PropertyAnimatorCoreSubsystem.cpp
@@ -113,7 +113,9 @@ bool UPropertyAnimatorCoreSubsystem::RegisterAnimatorClass(const UClass* InAnima
 
 	if (UPropertyAnimatorCoreBase* CDO = InAnimatorClass->GetDefaultObject<UPropertyAnimatorCoreBase>())
 	{
-		if (!CDO->GetAnimatorOriginalName().IsNone())
+		const TSharedPtr<const FPropertyAnimatorCoreMetadata> Metadata = CDO->GetAnimatorMetadata();
+
+		if (Metadata.IsValid() && !Metadata->Name.IsNone())
 		{
 			AnimatorsWeak.Add(CDO);
 
@@ -352,8 +354,9 @@ TSet<UPropertyAnimatorCoreBase*> UPropertyAnimatorCoreSubsystem::CreateAnimators
 	NewAnimators.Reserve(InActors.Num());
 
 #if WITH_EDITOR
+	const TSharedPtr<const FPropertyAnimatorCoreMetadata> AnimatorMetadata = AnimatorCDO->GetAnimatorMetadata();
 	const FText TransactionText = LOCTEXT("CreateAnimators", "Adding {0} animator to {1} actor(s)");
-	const FText AnimatorName = FText::FromName(AnimatorCDO->GetAnimatorOriginalName());
+	const FText AnimatorName = AnimatorMetadata->DisplayName;
 	const FText ActorCount = FText::FromString(FString::FromInt(InActors.Num()));
 
 	FScopedTransaction Transaction(FText::Format(TransactionText, AnimatorName, ActorCount), bInTransact && !GIsTransacting);
@@ -551,8 +554,8 @@ bool UPropertyAnimatorCoreSubsystem::ApplyAnimatorPreset(UPropertyAnimatorCoreBa
 	{
 #if WITH_EDITOR
 		const FText TransactionText = LOCTEXT("ApplyAnimatorPreset", "Applying {0} preset on {1} animator");
-		const FText PresetName = FText::FromString(InPreset->GetPresetDisplayName());
-		const FText AnimatorName = FText::FromName(InAnimator->GetAnimatorOriginalName());
+		const FText PresetName = InPreset->GetPresetDisplayName();
+		const FText AnimatorName = InAnimator->GetAnimatorMetadata()->DisplayName;
 
 		FScopedTransaction Transaction(FText::Format(TransactionText, PresetName, AnimatorName), bInTransact && !GIsTransacting);
 
@@ -576,8 +579,8 @@ bool UPropertyAnimatorCoreSubsystem::UnapplyAnimatorPreset(UPropertyAnimatorCore
 	{
 #if WITH_EDITOR
 		const FText TransactionText = LOCTEXT("UnapplyAnimatorPreset", "Unapplying {0} preset on {1} animator");
-		const FText PresetName = FText::FromString(InPreset->GetPresetDisplayName());
-		const FText AnimatorName = FText::FromName(InAnimator->GetAnimatorOriginalName());
+		const FText PresetName = InPreset->GetPresetDisplayName();
+		const FText AnimatorName = InAnimator->GetAnimatorMetadata()->DisplayName;
 
 		FScopedTransaction Transaction(FText::Format(TransactionText, PresetName, AnimatorName), bInTransact && !GIsTransacting);
 
@@ -611,7 +614,7 @@ bool UPropertyAnimatorCoreSubsystem::LinkAnimatorProperties(UPropertyAnimatorCor
 		if (InProperties.Num() == 1)
 		{
 			const FText PropertyName = FText::FromName(InProperties.Array()[0].GetLeafPropertyName());
-			const FText AnimatorName = FText::FromName(InAnimator->GetAnimatorOriginalName());
+			const FText AnimatorName = InAnimator->GetAnimatorMetadata()->DisplayName;
 
 			TransactionText = FText::Format(
 				LOCTEXT("LinkAnimatorProperty", "Linking {0} property to {1} animator")
@@ -622,7 +625,7 @@ bool UPropertyAnimatorCoreSubsystem::LinkAnimatorProperties(UPropertyAnimatorCor
 		else
 		{
 			const FText PropertyCount = FText::FromString(FString::FromInt(InProperties.Num()));
-			const FText AnimatorName = FText::FromName(InAnimator->GetAnimatorOriginalName());
+			const FText AnimatorName = InAnimator->GetAnimatorMetadata()->DisplayName;
 
 			TransactionText = FText::Format(
 				LOCTEXT("LinkAnimatorProperties", "Linking {0} properties to {1} animator")
@@ -669,7 +672,7 @@ bool UPropertyAnimatorCoreSubsystem::UnlinkAnimatorProperties(UPropertyAnimatorC
 		if (InProperties.Num() == 1)
 		{
 			const FText PropertyName = FText::FromName(InProperties.Array()[0].GetLeafPropertyName());
-			const FText AnimatorName = FText::FromName(InAnimator->GetAnimatorOriginalName());
+			const FText AnimatorName = InAnimator->GetAnimatorMetadata()->DisplayName;
 
 			TransactionText = FText::Format(
 				LOCTEXT("UnlinkAnimatorProperty", "Unlinking {0} property from {1} animator")
@@ -680,7 +683,7 @@ bool UPropertyAnimatorCoreSubsystem::UnlinkAnimatorProperties(UPropertyAnimatorC
 		else
 		{
 			const FText PropertyCount = FText::FromString(FString::FromInt(InProperties.Num()));
-			const FText AnimatorName = FText::FromName(InAnimator->GetAnimatorOriginalName());
+			const FText AnimatorName = InAnimator->GetAnimatorMetadata()->DisplayName;
 
 			TransactionText = FText::Format(
 				LOCTEXT("UnlinkAnimatorProperties", "Unlinking {0} properties from {1} animator")
diff --git a/Engine/Plugins/VirtualProduction/PropertyAnimatorCore/Source/PropertyAnimatorCore/Public/Animators/PropertyAnimatorCoreBase.h b/Engine/Plugins/VirtualProduction/PropertyAnimatorCore/Source/PropertyAnimatorCore/Public/Animators/PropertyAnimatorCoreBase.h
--- a/Engine/Plugins/VirtualProduction/PropertyAnimatorCore/Source/PropertyAnimatorCore/Public/Animators/PropertyAnimatorCoreBase.h
+++ b/Engine/Plugins/VirtualProduction/PropertyAnimatorCore/Source/PropertyAnimatorCore/Public/Animators/PropertyAnimatorCoreBase.h
@@ -23,7 +23,9 @@ enum class EPropertyAnimatorPropertySupport : uint8
 
 struct FPropertyAnimatorCoreMetadata
 {
-	FName Name;
+	FName Name = NAME_None;
+	FText DisplayName = FText::GetEmpty();
+	FText Description = FText::GetEmpty();
 	FName Category = TEXT("Default");
 };
 
@@ -121,10 +123,8 @@ public:
 		return AnimatorDisplayName;
 	}
 
-	PROPERTYANIMATORCORE_API FName GetAnimatorCategory() const;
-
-	/** Gets the Animator original name */
-	PROPERTYANIMATORCORE_API FName GetAnimatorOriginalName() const;
+	/** Gets the animator metadata */
+	PROPERTYANIMATORCORE_API TSharedRef<const FPropertyAnimatorCoreMetadata> GetAnimatorMetadata() const;
 
 	/** Get all linked properties within this animator */
 	PROPERTYANIMATORCORE_API TSet<FPropertyAnimatorCoreData> GetLinkedProperties() const;
diff --git a/Engine/Plugins/VirtualProduction/PropertyAnimatorCore/Source/PropertyAnimatorCore/Public/Presets/PropertyAnimatorCorePresetBase.h b/Engine/Plugins/VirtualProduction/PropertyAnimatorCore/Source/PropertyAnimatorCore/Public/Presets/PropertyAnimatorCorePresetBase.h
--- a/Engine/Plugins/VirtualProduction/PropertyAnimatorCore/Source/PropertyAnimatorCore/Public/Presets/PropertyAnimatorCorePresetBase.h
+++ b/Engine/Plugins/VirtualProduction/PropertyAnimatorCore/Source/PropertyAnimatorCore/Public/Presets/PropertyAnimatorCorePresetBase.h
@@ -33,7 +33,7 @@ public:
 		return PresetName;
 	}
 
-	PROPERTYANIMATORCORE_API FString GetPresetDisplayName() const;
+	PROPERTYANIMATORCORE_API FText GetPresetDisplayName() const;
 
 	/** Checks if the preset is supported on that actor and animator */
 	PROPERTYANIMATORCORE_API virtual bool IsPresetSupported(const AActor* InActor, const UPropertyAnimatorCoreBase* InAnimator) const PURE_VIRTUAL(UPropertyAnimatorCorePresetBase::IsPresetSupported, return false; );
@@ -66,10 +66,14 @@ public:
 	PROPERTYANIMATORCORE_API virtual TSharedRef<FPropertyAnimatorCorePresetArchiveImplementation> GetArchiveImplementation() const;
 
 protected:
-	/** Name used to display this preset to the user */
+	/** Name used to identify this preset */
 	UPROPERTY(VisibleInstanceOnly, Category="Animator")
 	FName PresetName;
 
+	/** Name used to display this preset to the user */
+	UPROPERTY(VisibleInstanceOnly, Category="Animator")
+	FText PresetDisplayName = FText::GetEmpty();
+
 	/** Version of this preset for diffs */
 	UPROPERTY(VisibleInstanceOnly, Category="Animator")
 	int32 PresetVersion = INDEX_NONE;
diff --git a/Engine/Plugins/VirtualProduction/PropertyAnimatorCore/Source/PropertyAnimatorCoreEditor/Private/Customizations/PropertyAnimatorCoreEditorStackCustomization.cpp b/Engine/Plugins/VirtualProduction/PropertyAnimatorCore/Source/PropertyAnimatorCoreEditor/Private/Customizations/PropertyAnimatorCoreEditorStackCustomization.cpp
--- a/Engine/Plugins/VirtualProduction/PropertyAnimatorCore/Source/PropertyAnimatorCoreEditor/Private/Customizations/PropertyAnimatorCoreEditorStackCustomization.cpp
+++ b/Engine/Plugins/VirtualProduction/PropertyAnimatorCore/Source/PropertyAnimatorCoreEditor/Private/Customizations/PropertyAnimatorCoreEditorStackCustomization.cpp
@@ -151,7 +151,7 @@ void UPropertyAnimatorCoreEditorStackCustomization::CustomizeStackHeader(const F
 			{
 				if (const UPropertyAnimatorCoreBase* Animator = SupportedItem->Get<UPropertyAnimatorCoreBase>(0))
 				{
-					PinnedAnimatorNames.Add(Animator->GetAnimatorOriginalName().ToString());
+					PinnedAnimatorNames.Add(Animator->GetAnimatorMetadata()->Name.ToString());
 				}
 			}
 		}
@@ -197,13 +197,13 @@ void UPropertyAnimatorCoreEditorStackCustomization::CustomizeItemHeader(const FO
 		// Commands for item on key events
 		const TSharedPtr<FUICommandList> ComponentCommands = CreateAnimatorCommands(InItem);
 
-		FString HeaderLabel = TEXT("Animators");
+		FText HeaderLabel = LOCTEXT("AnimatorComponentHeaderLabel", "Animators");
 		TAttribute<EOperatorStackEditorMessageType> MessageType = EOperatorStackEditorMessageType::None;
 		TAttribute<FText> MessageText = FText::GetEmpty();
 
 		if (InItem->GetValueCount() > 1)
 		{
-			HeaderLabel += TEXT(" (") + FString::FromInt(InItem->GetValueCount()) + TEXT(")");
+			HeaderLabel = FText::Format(LOCTEXT("AnimatorsHeaderLabel", "{0} ({1})"), HeaderLabel, FText::AsNumber(InItem->GetValueCount()));
 			MessageType = EOperatorStackEditorMessageType::Info;
 			MessageText = LOCTEXT("MultiAnimatorView", "You are viewing multiple items");
 		}
@@ -211,7 +211,7 @@ void UPropertyAnimatorCoreEditorStackCustomization::CustomizeItemHeader(const FO
 		InHeaderBuilder
 			.SetProperty(EnableProperty)
 			.SetIcon(ClassIcon.GetIcon())
-			.SetLabel(FText::FromString(HeaderLabel))
+			.SetLabel(HeaderLabel)
 			.SetToolbarMenu(HeaderAnimatorMenuName)
 			.SetContextMenu(ContextAnimatorMenuName)
 			.SetCommandList(ComponentCommands)
@@ -230,21 +230,23 @@ void UPropertyAnimatorCoreEditorStackCustomization::CustomizeItemHeader(const FO
 		const TSharedPtr<FUICommandList> AnimatorCommands = CreateAnimatorCommands(InItem);
 
 		const FString AnimatorDisplayName = Animator->GetAnimatorDisplayName().ToString();
+		const TSharedPtr<const FPropertyAnimatorCoreMetadata> AnimatorMetadata = Animator->GetAnimatorMetadata();
 
 		const TSet<FString> SearchKeywords
 		{
-			Animator->GetAnimatorOriginalName().ToString(),
+			AnimatorMetadata->Name.ToString(),
+			AnimatorMetadata->DisplayName.ToString(),
 			AnimatorDisplayName
 		};
 
 		/** Show last execution error messages if failed execution */
 		TAttribute<EOperatorStackEditorMessageType> MessageType = EOperatorStackEditorMessageType::None;
 		TAttribute<FText> MessageText = FText::GetEmpty();
-		FString HeaderLabel = TEXT("Animator");
+		FText HeaderLabel = FText::FromString(TEXT("Animator"));
 
 		if (InItem->GetValueCount() == 1)
 		{
-			HeaderLabel = AnimatorDisplayName;
+			HeaderLabel = AnimatorMetadata->DisplayName;
 
 			TWeakObjectPtr<UPropertyAnimatorCoreBase> AnimatorWeak(Animator);
 
@@ -272,8 +274,7 @@ void UPropertyAnimatorCoreEditorStackCustomization::CustomizeItemHeader(const FO
 		}
 		else
 		{
-			HeaderLabel = Animator->GetAnimatorOriginalName().ToString();
-			HeaderLabel += TEXT(" (") + FString::FromInt(InItem->GetValueCount()) + TEXT(")");
+			HeaderLabel = FText::Format(LOCTEXT("AnimatorHeaderLabel", "{0} ({1})"), AnimatorMetadata->DisplayName, FText::AsNumber(InItem->GetValueCount()));
 		}
 
 		static const FLinearColor AnimatorColor = FLinearColor(FColor::Orange).Desaturate(0.25);
@@ -284,7 +285,7 @@ void UPropertyAnimatorCoreEditorStackCustomization::CustomizeItemHeader(const FO
 			.SetSearchKeywords(SearchKeywords)
 			.SetExpandable(true)
 			.SetIcon(ClassIcon.GetIcon())
-			.SetLabel(FText::FromString(HeaderLabel))
+			.SetLabel(HeaderLabel)
 			.SetProperty(EnableProperty)
 			.SetCommandList(AnimatorCommands)
 			.SetToolbarMenu(HeaderAnimatorMenuName)
diff --git a/Engine/Plugins/VirtualProduction/PropertyAnimatorCore/Source/PropertyAnimatorCoreEditor/Private/Menus/PropertyAnimatorCoreEditorMenu.cpp b/Engine/Plugins/VirtualProduction/PropertyAnimatorCore/Source/PropertyAnimatorCoreEditor/Private/Menus/PropertyAnimatorCoreEditorMenu.cpp
--- a/Engine/Plugins/VirtualProduction/PropertyAnimatorCore/Source/PropertyAnimatorCoreEditor/Private/Menus/PropertyAnimatorCoreEditorMenu.cpp
+++ b/Engine/Plugins/VirtualProduction/PropertyAnimatorCore/Source/PropertyAnimatorCoreEditor/Private/Menus/PropertyAnimatorCoreEditorMenu.cpp
@@ -25,7 +25,8 @@ void UE::PropertyAnimatorCoreEditor::Menu::FillNewAnimatorSection(UToolMenu* InM
 	}
 
 	const UPropertyAnimatorCoreSubsystem* Subsystem = UPropertyAnimatorCoreSubsystem::Get();
-	if (!Subsystem)
+	const UPropertyAnimatorCoreEditorSubsystem* EditorSubsystem = UPropertyAnimatorCoreEditorSubsystem::Get();
+	if (!Subsystem || !EditorSubsystem)
 	{
 		return;
 	}
@@ -67,8 +68,9 @@ void UE::PropertyAnimatorCoreEditor::Menu::FillNewAnimatorSection(UToolMenu* InM
 		UPropertyAnimatorCoreBase* NewAnimator = NewPresetAnimator.Key;
 
 		const FName MenuName = NAME_None;
-		FName MenuCategory = NewAnimator->GetAnimatorCategory();
-		FText MenuLabel = FText::FromName(NewAnimator->GetAnimatorOriginalName());
+		const TSharedPtr<const FPropertyAnimatorCoreMetadata> Metadata = NewAnimator->GetAnimatorMetadata();
+		FName MenuCategory = Metadata->Category;
+		FText MenuLabel = Metadata->DisplayName;
 		const FText MenuTooltip = LOCTEXT("NewAnimator.Tooltip", "Create a new animator");
 		const FSlateIcon MenuIcon = FSlateIconFinder::FindIconForClass(NewAnimator->GetClass());
 
@@ -90,11 +92,12 @@ void UE::PropertyAnimatorCoreEditor::Menu::FillNewAnimatorSection(UToolMenu* InM
 			{
 				if (Preset)
 				{
-					MenuLabel = FText::FromString(Preset->GetPresetDisplayName() + TEXT(" (") + NewAnimator->GetAnimatorOriginalName().ToString() + TEXT(")"));
+					MenuLabel = FText::Format(LOCTEXT("PresetMenuLabel", "{0} ({1})"), Preset->GetPresetDisplayName(), Metadata->DisplayName);
 					MenuCategory = TEXT("Presets");
 				}
 
-				FToolMenuSection& AnimatorCategorySection = InMenu->FindOrAddSection(MenuCategory, FText::FromName(MenuCategory));
+				TSharedPtr<const FPropertyAnimatorCoreEditorCategoryMetadata> AnimatorCategory = EditorSubsystem->FindAnimatorCategory(MenuCategory);
+				FToolMenuSection& AnimatorCategorySection = InMenu->FindOrAddSection(AnimatorCategory ? AnimatorCategory->Name : MenuCategory, AnimatorCategory ? AnimatorCategory->DisplayName : FText::FromName(MenuCategory));
 
 				AnimatorCategorySection.AddMenuEntry(
 					MenuName
@@ -472,7 +475,7 @@ void UE::PropertyAnimatorCoreEditor::Menu::ExecuteApplyLastCreatedAnimatorPreset
 
 	FScopedTransaction Transaction(
 		FText::Format(LOCTEXT("ApplyAnimatorsPreset", "Toggle apply preset {0} on {1} animator(s)"),
-			FText::FromString(InPreset->GetPresetDisplayName()),
+			InPreset->GetPresetDisplayName(),
 			FText::AsNumber(LastAnimators.Num()))
 		, bShouldTransact && !GIsTransacting);
 
@@ -512,7 +515,7 @@ void UE::PropertyAnimatorCoreEditor::Menu::ExecuteLinkAnimatorPresetAction(const
 
 	FScopedTransaction Transaction(
 		FText::Format(LOCTEXT("ApplyAnimatorsPreset", "Toggle apply preset {0} on {1} animator(s)"),
-			FText::FromString(InPreset->GetPresetDisplayName()),
+			InPreset->GetPresetDisplayName(),
 			FText::AsNumber(InAnimators.Num()))
 		, bShouldTransact && !GIsTransacting);
 
@@ -838,7 +841,7 @@ void UE::PropertyAnimatorCoreEditor::Menu::FillNewAnimatorSubmenu(UToolMenu* InM
 			}
 
 			const FString MenuName = SupportedPreset->GetPresetName().ToString();
-			const FText MenuLabel = FText::FromString(SupportedPreset->GetPresetDisplayName());
+			const FText MenuLabel = SupportedPreset->GetPresetDisplayName();
 
 			PresetSection.AddSubMenu(
 				FName(TEXT("Create") + MenuName)
@@ -940,7 +943,7 @@ void UE::PropertyAnimatorCoreEditor::Menu::FillLinkAnimatorSubmenu(UToolMenu* In
 		}
 
 		const FName MenuName = SupportedPreset->GetPresetName();
-		const FText MenuLabel = FText::FromString(SupportedPreset->GetPresetDisplayName());
+		const FText MenuLabel = SupportedPreset->GetPresetDisplayName();
 		const FText MenuTooltip = LOCTEXT("LinkAnimatorPresetSection.Tooltip", "Link or unlink a preset from this animator");
 
 		PresetSection.AddSubMenu(
diff --git a/Engine/Plugins/VirtualProduction/PropertyAnimatorCore/Source/PropertyAnimatorCoreEditor/Private/Sequencer/MovieSceneAnimatorTrackEditor.cpp b/Engine/Plugins/VirtualProduction/PropertyAnimatorCore/Source/PropertyAnimatorCoreEditor/Private/Sequencer/MovieSceneAnimatorTrackEditor.cpp
--- a/Engine/Plugins/VirtualProduction/PropertyAnimatorCore/Source/PropertyAnimatorCoreEditor/Private/Sequencer/MovieSceneAnimatorTrackEditor.cpp
+++ b/Engine/Plugins/VirtualProduction/PropertyAnimatorCore/Source/PropertyAnimatorCoreEditor/Private/Sequencer/MovieSceneAnimatorTrackEditor.cpp
@@ -133,7 +133,7 @@ void FMovieSceneAnimatorTrackEditor::ExecuteAddTrack(const TArray<UObject*>& InO
 
 			if (const UPropertyAnimatorCoreBase* Animator = Cast<UPropertyAnimatorCoreBase>(Owner->GetOuter()))
 			{
-				NewTrack->SetDisplayName(FText::Format(LOCTEXT("MovieSceneAnimatorTrackName", "Animator {0} Track"), FText::FromName(Animator->GetAnimatorOriginalName())));
+				NewTrack->SetDisplayName(FText::Format(LOCTEXT("MovieSceneAnimatorTrackName", "Animator {0} Track"), Animator->GetAnimatorMetadata()->DisplayName));
 			}
 			else
 			{
diff --git a/Engine/Plugins/VirtualProduction/PropertyAnimatorCore/Source/PropertyAnimatorCoreEditor/Private/Subsystems/PropertyAnimatorCoreEditorSubsystem.cpp b/Engine/Plugins/VirtualProduction/PropertyAnimatorCore/Source/PropertyAnimatorCoreEditor/Private/Subsystems/PropertyAnimatorCoreEditorSubsystem.cpp
--- a/Engine/Plugins/VirtualProduction/PropertyAnimatorCore/Source/PropertyAnimatorCoreEditor/Private/Subsystems/PropertyAnimatorCoreEditorSubsystem.cpp
+++ b/Engine/Plugins/VirtualProduction/PropertyAnimatorCore/Source/PropertyAnimatorCoreEditor/Private/Subsystems/PropertyAnimatorCoreEditorSubsystem.cpp
@@ -208,6 +208,33 @@ UPropertyAnimatorCorePresetBase* UPropertyAnimatorCoreEditorSubsystem::CreatePre
 	return NewPreset;
 }
 
+bool UPropertyAnimatorCoreEditorSubsystem::RegisterAnimatorCategory(const FPropertyAnimatorCoreEditorCategoryMetadata& InCategoryMetadata)
+{
+	if (InCategoryMetadata.Name.IsNone())
+	{
+		return false;
+	}
+
+	if (AnimatorCategories.Contains(InCategoryMetadata.Name))
+	{
+		return false;
+	}
+
+	AnimatorCategories.Emplace(InCategoryMetadata.Name, MakeShared<const FPropertyAnimatorCoreEditorCategoryMetadata>(InCategoryMetadata));
+
+	return true;
+}
+
+TSharedPtr<const FPropertyAnimatorCoreEditorCategoryMetadata> UPropertyAnimatorCoreEditorSubsystem::FindAnimatorCategory(FName InCategoryIdentifier) const
+{
+	if (const TSharedRef<const FPropertyAnimatorCoreEditorCategoryMetadata>* Category = AnimatorCategories.Find(InCategoryIdentifier))
+	{
+		return *Category;
+	}
+	
+	return nullptr;
+}
+
 UPropertyAnimatorCoreEditorSubsystem* UPropertyAnimatorCoreEditorSubsystem::Get()
 {
 	if (GEditor)
@@ -237,6 +264,41 @@ void UPropertyAnimatorCoreEditorSubsystem::Initialize(FSubsystemCollectionBase&
 	// Initialize here, to setup ClassIcons if no other code uses it
 	FPropertyAnimatorCoreEditorStyle::Get();
 
+	// Register animator categories
+	{
+		const FPropertyAnimatorCoreEditorCategoryMetadata DefaultCategory
+		{
+			.Name = TEXT("Default"),
+			.DisplayName = LOCTEXT("DefaultCategoryDisplayName", "Default")
+		};
+
+		RegisterAnimatorCategory(DefaultCategory);
+
+		const FPropertyAnimatorCoreEditorCategoryMetadata PresetCategory
+		{
+			.Name = TEXT("Presets"),
+			.DisplayName = LOCTEXT("PresetCategoryDisplayName", "Presets")
+		};
+
+		RegisterAnimatorCategory(PresetCategory);
+
+		const FPropertyAnimatorCoreEditorCategoryMetadata TextCategory
+		{
+			.Name = TEXT("Text"),
+			.DisplayName = LOCTEXT("TextCategoryDisplayName", "Text")
+		};
+
+		RegisterAnimatorCategory(TextCategory);
+
+		const FPropertyAnimatorCoreEditorCategoryMetadata NumericCategory
+		{
+			.Name = TEXT("Numeric"),
+			.DisplayName = LOCTEXT("NumericCategoryDisplayName", "Numeric")
+		};
+
+		RegisterAnimatorCategory(NumericCategory);
+	}
+
 	RegisterDetailPanelCustomization();
 }
 
diff --git a/Engine/Plugins/VirtualProduction/PropertyAnimatorCore/Source/PropertyAnimatorCoreEditor/Public/Subsystems/PropertyAnimatorCoreEditorSubsystem.h b/Engine/Plugins/VirtualProduction/PropertyAnimatorCore/Source/PropertyAnimatorCoreEditor/Public/Subsystems/PropertyAnimatorCoreEditorSubsystem.h
--- a/Engine/Plugins/VirtualProduction/PropertyAnimatorCore/Source/PropertyAnimatorCoreEditor/Public/Subsystems/PropertyAnimatorCoreEditorSubsystem.h
+++ b/Engine/Plugins/VirtualProduction/PropertyAnimatorCore/Source/PropertyAnimatorCoreEditor/Public/Subsystems/PropertyAnimatorCoreEditorSubsystem.h
@@ -22,6 +22,12 @@ struct FPropertyAnimatorCoreEditorEditPanelOptions;
 struct FOnGenerateGlobalRowExtensionArgs;
 struct FPropertyRowExtensionButton;
 
+struct FPropertyAnimatorCoreEditorCategoryMetadata
+{
+	FName Name;
+	FText DisplayName;
+};
+
 /** Singleton class that handles editor operations for property animators */
 UCLASS()
 class UPropertyAnimatorCoreEditorSubsystem : public UEditorSubsystem
@@ -43,6 +49,12 @@ public:
 	/** Creates a preset asset for an item */
 	PROPERTYANIMATORCOREEDITOR_API UPropertyAnimatorCorePresetBase* CreatePresetAsset(TSubclassOf<UPropertyAnimatorCorePresetBase> InPresetClass, const TArray<IPropertyAnimatorCorePresetable*>& InPresetables);
 
+	/** Register a category for animator */
+	PROPERTYANIMATORCOREEDITOR_API bool RegisterAnimatorCategory(const FPropertyAnimatorCoreEditorCategoryMetadata& InCategoryMetadata);
+
+	/** Finds a registered category for animator */
+	PROPERTYANIMATORCOREEDITOR_API TSharedPtr<const FPropertyAnimatorCoreEditorCategoryMetadata> FindAnimatorCategory(FName InCategoryIdentifier) const;
+
 protected:
 	/** Setup details panel button customization */
 	void RegisterDetailPanelCustomization();
@@ -85,4 +97,6 @@ protected:
 	TWeakPtr<SPropertyAnimatorCoreEditorEditPanel> PropertyControllerPanelWeak;
 
 	TSharedPtr<FPropertyAnimatorCoreEditorMenuData> LastMenuData;
+
+	TMap<FName, TSharedRef<const FPropertyAnimatorCoreEditorCategoryMetadata>> AnimatorCategories;
 };
\ No newline at end of file
