MotionDesign : PropertyAnimator
- Added reset button to player controls that pauses and sets time to 0 for manual time source

diff --git a/Engine/Plugins/VirtualProduction/PropertyAnimatorCore/Source/PropertyAnimatorCore/Private/TimeSources/PropertyAnimatorCoreManualTimeSource.cpp b/Engine/Plugins/VirtualProduction/PropertyAnimatorCore/Source/PropertyAnimatorCore/Private/TimeSources/PropertyAnimatorCoreManualTimeSource.cpp
--- a/Engine/Plugins/VirtualProduction/PropertyAnimatorCore/Source/PropertyAnimatorCore/Private/TimeSources/PropertyAnimatorCoreManualTimeSource.cpp
+++ b/Engine/Plugins/VirtualProduction/PropertyAnimatorCore/Source/PropertyAnimatorCore/Private/TimeSources/PropertyAnimatorCoreManualTimeSource.cpp
@@ -5,6 +5,11 @@
 #include "Misc/App.h"
 #include "Presets/PropertyAnimatorCorePresetArchive.h"
 
+FName UPropertyAnimatorCoreManualTimeSource::GetCustomTimePropertyName()
+{
+	return GET_MEMBER_NAME_CHECKED(UPropertyAnimatorCoreManualTimeSource, CustomTime);
+}
+
 void UPropertyAnimatorCoreManualTimeSource::SetCustomTime(double InTime)
 {
 	if (FMath::IsNearlyEqual(InTime, CustomTime))
diff --git a/Engine/Plugins/VirtualProduction/PropertyAnimatorCore/Source/PropertyAnimatorCore/Public/TimeSources/PropertyAnimatorCoreManualTimeSource.h b/Engine/Plugins/VirtualProduction/PropertyAnimatorCore/Source/PropertyAnimatorCore/Public/TimeSources/PropertyAnimatorCoreManualTimeSource.h
--- a/Engine/Plugins/VirtualProduction/PropertyAnimatorCore/Source/PropertyAnimatorCore/Public/TimeSources/PropertyAnimatorCoreManualTimeSource.h
+++ b/Engine/Plugins/VirtualProduction/PropertyAnimatorCore/Source/PropertyAnimatorCore/Public/TimeSources/PropertyAnimatorCoreManualTimeSource.h
@@ -24,6 +24,8 @@ class UPropertyAnimatorCoreManualTimeSource : public UPropertyAnimatorCoreTimeSo
 	GENERATED_BODY()
 
 public:
+	static PROPERTYANIMATORCORE_API FName GetCustomTimePropertyName();
+
 	UPropertyAnimatorCoreManualTimeSource()
 		: UPropertyAnimatorCoreTimeSourceBase(TEXT("Manual"))
 	{}
diff --git a/Engine/Plugins/VirtualProduction/PropertyAnimatorCore/Source/PropertyAnimatorCoreEditor/Private/Customizations/PropertyAnimatorCoreEditorManualStateTypeCustomization.cpp b/Engine/Plugins/VirtualProduction/PropertyAnimatorCore/Source/PropertyAnimatorCoreEditor/Private/Customizations/PropertyAnimatorCoreEditorManualStateTypeCustomization.cpp
--- a/Engine/Plugins/VirtualProduction/PropertyAnimatorCore/Source/PropertyAnimatorCoreEditor/Private/Customizations/PropertyAnimatorCoreEditorManualStateTypeCustomization.cpp
+++ b/Engine/Plugins/VirtualProduction/PropertyAnimatorCore/Source/PropertyAnimatorCoreEditor/Private/Customizations/PropertyAnimatorCoreEditorManualStateTypeCustomization.cpp
@@ -19,6 +19,11 @@ void FPropertyAnimatorCoreEditorManualStateTypeCustomization::CustomizeHeader(TS
 
 	StatusPropertyHandle = InPropertyHandle;
 
+	if (const TSharedPtr<IPropertyHandle> ParentHandle = InPropertyHandle->GetParentHandle())
+	{
+		CustomTimePropertyHandle = ParentHandle->GetChildHandle(UPropertyAnimatorCoreManualTimeSource::GetCustomTimePropertyName(), /** Recurse */false);
+	}
+
 	InRow.NameContent()
 	[
 		InPropertyHandle->CreatePropertyNameWidget()
@@ -94,6 +99,22 @@ void FPropertyAnimatorCoreEditorManualStateTypeCustomization::CustomizeHeader(TS
 				.Image(FPropertyAnimatorCoreEditorStyle::Get().GetBrush(TEXT("ManualTimeSourceControl.Stop")))
 			]
 		]
+		+ SHorizontalBox::Slot()
+		.AutoWidth()
+		.Padding(5.f, 0.f, 0.f, 5.f)
+		.HAlign(HAlign_Fill)
+		[
+			SNew(SButton)
+			.HAlign(HAlign_Fill)
+			.ButtonColorAndOpacity(FStyleColors::White)
+			.OnClicked(this, &FPropertyAnimatorCoreEditorManualStateTypeCustomization::ResetPlaybackStatus)
+			.Visibility(CustomTimePropertyHandle.IsValid() ? EVisibility::Visible : EVisibility::Collapsed)
+			[
+				SNew(SImage)
+				.DesiredSizeOverride(ImageSize)
+				.Image(FPropertyAnimatorCoreEditorStyle::Get().GetBrush(TEXT("ManualTimeSourceControl.Reset")))
+			]
+		]
 	];
 }
 
@@ -111,6 +132,21 @@ FReply FPropertyAnimatorCoreEditorManualStateTypeCustomization::SetPlaybackStatu
 	return FReply::Handled();
 }
 
+FReply FPropertyAnimatorCoreEditorManualStateTypeCustomization::ResetPlaybackStatus()
+{
+	if (StatusPropertyHandle.IsValid() && StatusPropertyHandle->IsValidHandle())
+	{
+		StatusPropertyHandle->SetValue(static_cast<uint8>(EPropertyAnimatorCoreManualStatus::Paused));
+	}
+
+	if (CustomTimePropertyHandle.IsValid() && CustomTimePropertyHandle->IsValidHandle())
+	{
+		CustomTimePropertyHandle->SetValue(0.0);
+	}
+
+	return FReply::Handled();
+}
+
 bool FPropertyAnimatorCoreEditorManualStateTypeCustomization::IsPlaybackStatusAllowed(EPropertyAnimatorCoreManualStatus InStatus) const
 {
 	if (StatusPropertyHandle.IsValid() && StatusPropertyHandle->IsValidHandle())
diff --git a/Engine/Plugins/VirtualProduction/PropertyAnimatorCore/Source/PropertyAnimatorCoreEditor/Private/Customizations/PropertyAnimatorCoreEditorManualStateTypeCustomization.h b/Engine/Plugins/VirtualProduction/PropertyAnimatorCore/Source/PropertyAnimatorCoreEditor/Private/Customizations/PropertyAnimatorCoreEditorManualStateTypeCustomization.h
--- a/Engine/Plugins/VirtualProduction/PropertyAnimatorCore/Source/PropertyAnimatorCoreEditor/Private/Customizations/PropertyAnimatorCoreEditorManualStateTypeCustomization.h
+++ b/Engine/Plugins/VirtualProduction/PropertyAnimatorCore/Source/PropertyAnimatorCoreEditor/Private/Customizations/PropertyAnimatorCoreEditorManualStateTypeCustomization.h
@@ -26,8 +26,10 @@ public:
 
 protected:
 	FReply SetPlaybackStatus(EPropertyAnimatorCoreManualStatus InStatus);
+	FReply ResetPlaybackStatus();
 	bool IsPlaybackStatusAllowed(EPropertyAnimatorCoreManualStatus InStatus) const;
 	FSlateColor IsPlaybackStatusActive(EPropertyAnimatorCoreManualStatus InStatus) const;
 
 	TSharedPtr<IPropertyHandle> StatusPropertyHandle;
+	TSharedPtr<IPropertyHandle> CustomTimePropertyHandle;
 };
diff --git a/Engine/Plugins/VirtualProduction/PropertyAnimatorCore/Source/PropertyAnimatorCoreEditor/Private/Styles/PropertyAnimatorCoreEditorStyle.cpp b/Engine/Plugins/VirtualProduction/PropertyAnimatorCore/Source/PropertyAnimatorCoreEditor/Private/Styles/PropertyAnimatorCoreEditorStyle.cpp
--- a/Engine/Plugins/VirtualProduction/PropertyAnimatorCore/Source/PropertyAnimatorCoreEditor/Private/Styles/PropertyAnimatorCoreEditorStyle.cpp
+++ b/Engine/Plugins/VirtualProduction/PropertyAnimatorCore/Source/PropertyAnimatorCoreEditor/Private/Styles/PropertyAnimatorCoreEditorStyle.cpp
@@ -30,6 +30,7 @@ FPropertyAnimatorCoreEditorStyle::FPropertyAnimatorCoreEditorStyle()
 	Set("ManualTimeSourceControl.PlayBackward",     new IMAGE_BRUSH_SVG("EditorIcons/PlayReverse",                 Icon16x16));
 	Set("ManualTimeSourceControl.Pause",            new IMAGE_BRUSH_SVG("EditorIcons/Pause",                       Icon16x16));
 	Set("ManualTimeSourceControl.Stop",             new IMAGE_BRUSH_SVG("EditorIcons/Stop",                        Icon16x16));
+	Set("ManualTimeSourceControl.Reset",            new IMAGE_BRUSH("EditorIcons/Reset",                           Icon16x16));
 
 	FSlateStyleRegistry::RegisterSlateStyle(*this);
 }
