MotionDesign : PropertyAnimator
- Fixed text character index is capped due to property being restored to its original value
- Fixed text would duplicate itself when an additive animator is bound to it and value is modified in details panel
- Fixed text flickers due to property being restored then updated with newer value in a short amount of time

diff --git a/Engine/Plugins/VirtualProduction/PropertyAnimatorCore/Source/PropertyAnimatorCore/Private/Animators/PropertyAnimatorCoreBase.cpp b/Engine/Plugins/VirtualProduction/PropertyAnimatorCore/Source/PropertyAnimatorCore/Private/Animators/PropertyAnimatorCoreBase.cpp
--- a/Engine/Plugins/VirtualProduction/PropertyAnimatorCore/Source/PropertyAnimatorCore/Private/Animators/PropertyAnimatorCoreBase.cpp
+++ b/Engine/Plugins/VirtualProduction/PropertyAnimatorCore/Source/PropertyAnimatorCore/Private/Animators/PropertyAnimatorCoreBase.cpp
@@ -178,8 +178,7 @@ void UPropertyAnimatorCoreBase::PreDuplicate(FObjectDuplicationParameters& InPar
 {
 	Super::PreDuplicate(InParams);
 
-	constexpr bool bForceReset = true;
-	RestoreProperties(bForceReset);
+	RestoreProperties();
 }
 
 void UPropertyAnimatorCoreBase::PostDuplicate(EDuplicateMode::Type InMode)
@@ -195,16 +194,14 @@ void UPropertyAnimatorCoreBase::PreEditUndo()
 {
 	Super::PreEditUndo();
 
-	constexpr bool bForceReset = true;
-	RestoreProperties(bForceReset);
+	RestoreProperties();
 }
 
 void UPropertyAnimatorCoreBase::PostEditUndo()
 {
 	Super::PostEditUndo();
 
-	constexpr bool bForceReset = true;
-	RestoreProperties(bForceReset);
+	RestoreProperties();
 }
 
 void UPropertyAnimatorCoreBase::PreEditChange(FProperty* InPropertyAboutToChange)
@@ -214,8 +211,7 @@ void UPropertyAnimatorCoreBase::PreEditChange(FProperty* InPropertyAboutToChange
 	if (InPropertyAboutToChange
 		&& InPropertyAboutToChange->GetFName() == GET_MEMBER_NAME_CHECKED(UPropertyAnimatorCoreBase, bAnimatorEnabled))
 	{
-		constexpr bool bForceReset = true;
-		RestoreProperties(bForceReset);
+		RestoreProperties();
 	}
 }
 
@@ -516,8 +512,7 @@ void UPropertyAnimatorCoreBase::OnAnimatorDisabled(EPropertyAnimatorCoreUpdateEv
 		, *GetAnimatorMetadata()->Name.ToString()
 		, InType);
 
-	constexpr bool bForceReset = true;
-	RestoreProperties(bForceReset);
+	RestoreProperties();
 }
 
 TSubclassOf<UPropertyAnimatorCoreContext> UPropertyAnimatorCoreBase::GetPropertyContextClass(const FPropertyAnimatorCoreData& InProperty)
@@ -570,7 +565,7 @@ void UPropertyAnimatorCoreBase::OnTimeSourceNameChanged()
 
 void UPropertyAnimatorCoreBase::OnTimeSourceEnterIdleState()
 {
-	RestoreProperties(/** Force */true);
+	RestoreProperties();
 }
 
 void UPropertyAnimatorCoreBase::ResolvePropertiesOwner(AActor* InNewOwner)
@@ -598,8 +593,6 @@ void UPropertyAnimatorCoreBase::ResolvePropertiesOwner(AActor* InNewOwner)
 
 void UPropertyAnimatorCoreBase::EvaluateAnimator(FInstancedPropertyBag& InParameters)
 {
-	RestoreProperties();
-
 	SaveProperties();
 
 	EvaluatedPropertyValues.Reset();
@@ -636,37 +629,16 @@ void UPropertyAnimatorCoreBase::OnPreSaveWorld(UWorld* InWorld, FObjectPreSaveCo
 		return;
 	}
 
-	constexpr bool bForce = true;
-	RestoreProperties(bForce);
+	RestoreProperties();
 }
 #endif
 
-void UPropertyAnimatorCoreBase::RestoreProperties(bool bInForce)
+void UPropertyAnimatorCoreBase::RestoreProperties()
 {
 	constexpr bool bResolve = false;
-	ForEachLinkedProperty<UPropertyAnimatorCoreContext>([bInForce](UPropertyAnimatorCoreContext* InOptions, const FPropertyAnimatorCoreData& InPropertyData)
+	bool bResult = ForEachLinkedProperty<UPropertyAnimatorCoreContext>([](UPropertyAnimatorCoreContext* InOptions, const FPropertyAnimatorCoreData& InPropertyData)
 	{
-		bool bRestore = false;
-
-		if (bInForce)
-		{
-			bRestore = true;
-		}
-		if (InOptions->Mode == EPropertyAnimatorCoreMode::Additive)
-		{
-			bRestore = true;
-		}
-		else if (InOptions->Mode == EPropertyAnimatorCoreMode::Absolute
-			&& (InOptions->IsResolvable() || InOptions->IsConverted()))
-		{
-			bRestore = true;
-		}
-
-		if (bRestore)
-		{
-			InOptions->Restore();
-		}
-
+		InOptions->Restore();
 		return true;
 	}, bResolve);
 }
@@ -674,7 +646,7 @@ void UPropertyAnimatorCoreBase::RestoreProperties(bool bInForce)
 void UPropertyAnimatorCoreBase::SaveProperties()
 {
 	constexpr bool bResolve = false;
-	ForEachLinkedProperty<UPropertyAnimatorCoreContext>([](UPropertyAnimatorCoreContext* InOptions, const FPropertyAnimatorCoreData& InPropertyData)
+	bool bResult = ForEachLinkedProperty<UPropertyAnimatorCoreContext>([](UPropertyAnimatorCoreContext* InOptions, const FPropertyAnimatorCoreData& InPropertyData)
 	{
 		InOptions->Save();
 		return true;
diff --git a/Engine/Plugins/VirtualProduction/PropertyAnimatorCore/Source/PropertyAnimatorCore/Private/Components/PropertyAnimatorCoreComponent.cpp b/Engine/Plugins/VirtualProduction/PropertyAnimatorCore/Source/PropertyAnimatorCore/Private/Components/PropertyAnimatorCoreComponent.cpp
--- a/Engine/Plugins/VirtualProduction/PropertyAnimatorCore/Source/PropertyAnimatorCore/Private/Components/PropertyAnimatorCoreComponent.cpp
+++ b/Engine/Plugins/VirtualProduction/PropertyAnimatorCore/Source/PropertyAnimatorCore/Private/Components/PropertyAnimatorCoreComponent.cpp
@@ -235,13 +235,11 @@ void UPropertyAnimatorCoreComponent::OnComponentCreated()
 	if (AActor* OwningActor = GetOwner())
 	{
 		// For spawnable templates, restore and resolve properties owner
-		constexpr bool bForceRestore = true;
-
 		for (const TObjectPtr<UPropertyAnimatorCoreBase>& Animator : PropertyAnimators)
 		{
 			if (Animator)
 			{
-				Animator->RestoreProperties(bForceRestore);
+				Animator->RestoreProperties();
 				Animator->ResolvePropertiesOwner(OwningActor);
 			}
 		}
diff --git a/Engine/Plugins/VirtualProduction/PropertyAnimatorCore/Source/PropertyAnimatorCore/Private/Properties/Handlers/PropertyAnimatorCoreDoubleHandler.cpp b/Engine/Plugins/VirtualProduction/PropertyAnimatorCore/Source/PropertyAnimatorCore/Private/Properties/Handlers/PropertyAnimatorCoreDoubleHandler.cpp
--- a/Engine/Plugins/VirtualProduction/PropertyAnimatorCore/Source/PropertyAnimatorCore/Private/Properties/Handlers/PropertyAnimatorCoreDoubleHandler.cpp
+++ b/Engine/Plugins/VirtualProduction/PropertyAnimatorCore/Source/PropertyAnimatorCore/Private/Properties/Handlers/PropertyAnimatorCoreDoubleHandler.cpp
@@ -68,6 +68,22 @@ bool UPropertyAnimatorCoreDoubleHandler::AddValue(const FPropertyAnimatorCoreDat
 	return true;
 }
 
+bool UPropertyAnimatorCoreDoubleHandler::AddValue(const FPropertyAnimatorCoreData& InPropertyData, const FInstancedPropertyBag& InOriginalValue, const FInstancedPropertyBag& InDeltaValue)
+{
+	const FName PropertyHash(InPropertyData.GetLocatorPathHash());
+	const TValueOrError<double, EPropertyBagResult> OriginalValueResult = InOriginalValue.GetValueDouble(PropertyHash);
+	const TValueOrError<double, EPropertyBagResult> DeltaValueResult = InDeltaValue.GetValueDouble(PropertyHash);
+	if (!OriginalValueResult.HasValue() || !DeltaValueResult.HasValue())
+	{
+		return false;
+	}
+
+	double NewValue = OriginalValueResult.GetValue() + DeltaValueResult.GetValue();
+	InPropertyData.SetPropertyValuePtr(&NewValue);
+
+	return true;
+}
+
 bool UPropertyAnimatorCoreDoubleHandler::SubtractValue(const FPropertyAnimatorCoreData& InPropertyData, const FInstancedPropertyBag& InValue)
 {
 	const FName PropertyHash(InPropertyData.GetLocatorPathHash());
diff --git a/Engine/Plugins/VirtualProduction/PropertyAnimatorCore/Source/PropertyAnimatorCore/Private/Properties/Handlers/PropertyAnimatorCoreDoubleHandler.h b/Engine/Plugins/VirtualProduction/PropertyAnimatorCore/Source/PropertyAnimatorCore/Private/Properties/Handlers/PropertyAnimatorCoreDoubleHandler.h
--- a/Engine/Plugins/VirtualProduction/PropertyAnimatorCore/Source/PropertyAnimatorCore/Private/Properties/Handlers/PropertyAnimatorCoreDoubleHandler.h
+++ b/Engine/Plugins/VirtualProduction/PropertyAnimatorCore/Source/PropertyAnimatorCore/Private/Properties/Handlers/PropertyAnimatorCoreDoubleHandler.h
@@ -18,6 +18,7 @@ public:
 	virtual bool SetValue(const FPropertyAnimatorCoreData& InPropertyData, const FInstancedPropertyBag& InValue) override;
 	virtual bool IsAdditiveSupported() const override;
 	virtual bool AddValue(const FPropertyAnimatorCoreData& InPropertyData, const FInstancedPropertyBag& InValue) override;
+	virtual bool AddValue(const FPropertyAnimatorCoreData& InPropertyData, const FInstancedPropertyBag& InOriginalValue, const FInstancedPropertyBag& InDeltaValue) override;
 	virtual bool SubtractValue(const FPropertyAnimatorCoreData& InPropertyData, const FInstancedPropertyBag& InValue) override;
 	virtual bool GetDefaultValue(const FPropertyAnimatorCoreData& InPropertyData, FInstancedPropertyBag& OutValue) override;
 	//~ End UPropertyAnimatorCoreHandlerBase
diff --git a/Engine/Plugins/VirtualProduction/PropertyAnimatorCore/Source/PropertyAnimatorCore/Private/Properties/Handlers/PropertyAnimatorCoreFloatHandler.cpp b/Engine/Plugins/VirtualProduction/PropertyAnimatorCore/Source/PropertyAnimatorCore/Private/Properties/Handlers/PropertyAnimatorCoreFloatHandler.cpp
--- a/Engine/Plugins/VirtualProduction/PropertyAnimatorCore/Source/PropertyAnimatorCore/Private/Properties/Handlers/PropertyAnimatorCoreFloatHandler.cpp
+++ b/Engine/Plugins/VirtualProduction/PropertyAnimatorCore/Source/PropertyAnimatorCore/Private/Properties/Handlers/PropertyAnimatorCoreFloatHandler.cpp
@@ -68,6 +68,22 @@ bool UPropertyAnimatorCoreFloatHandler::AddValue(const FPropertyAnimatorCoreData
 	return true;
 }
 
+bool UPropertyAnimatorCoreFloatHandler::AddValue(const FPropertyAnimatorCoreData& InPropertyData, const FInstancedPropertyBag& InOriginalValue, const FInstancedPropertyBag& InDeltaValue)
+{
+	const FName PropertyHash(InPropertyData.GetLocatorPathHash());
+	const TValueOrError<float, EPropertyBagResult> OriginalValueResult = InOriginalValue.GetValueFloat(PropertyHash);
+	const TValueOrError<float, EPropertyBagResult> DeltaValueResult = InDeltaValue.GetValueFloat(PropertyHash);
+	if (!OriginalValueResult.HasValue() || !DeltaValueResult.HasValue())
+	{
+		return false;
+	}
+
+	float NewValue = OriginalValueResult.GetValue() + DeltaValueResult.GetValue();
+	InPropertyData.SetPropertyValuePtr(&NewValue);
+
+	return true;
+}
+
 bool UPropertyAnimatorCoreFloatHandler::SubtractValue(const FPropertyAnimatorCoreData& InPropertyData, const FInstancedPropertyBag& InValue)
 {
 	const FName PropertyHash(InPropertyData.GetLocatorPathHash());
diff --git a/Engine/Plugins/VirtualProduction/PropertyAnimatorCore/Source/PropertyAnimatorCore/Private/Properties/Handlers/PropertyAnimatorCoreFloatHandler.h b/Engine/Plugins/VirtualProduction/PropertyAnimatorCore/Source/PropertyAnimatorCore/Private/Properties/Handlers/PropertyAnimatorCoreFloatHandler.h
--- a/Engine/Plugins/VirtualProduction/PropertyAnimatorCore/Source/PropertyAnimatorCore/Private/Properties/Handlers/PropertyAnimatorCoreFloatHandler.h
+++ b/Engine/Plugins/VirtualProduction/PropertyAnimatorCore/Source/PropertyAnimatorCore/Private/Properties/Handlers/PropertyAnimatorCoreFloatHandler.h
@@ -18,6 +18,7 @@ public:
 	virtual bool SetValue(const FPropertyAnimatorCoreData& InPropertyData, const FInstancedPropertyBag& InValue) override;
 	virtual bool IsAdditiveSupported() const override;
 	virtual bool AddValue(const FPropertyAnimatorCoreData& InPropertyData, const FInstancedPropertyBag& InValue) override;
+	virtual bool AddValue(const FPropertyAnimatorCoreData& InPropertyData, const FInstancedPropertyBag& InOriginalValue, const FInstancedPropertyBag& InDeltaValue) override;
 	virtual bool SubtractValue(const FPropertyAnimatorCoreData& InPropertyData, const FInstancedPropertyBag& InValue) override;
 	virtual bool GetDefaultValue(const FPropertyAnimatorCoreData& InPropertyData, FInstancedPropertyBag& OutValue) override;
 	//~ End UPropertyAnimatorCoreHandlerBase
diff --git a/Engine/Plugins/VirtualProduction/PropertyAnimatorCore/Source/PropertyAnimatorCore/Private/Properties/Handlers/PropertyAnimatorCoreIntHandler.cpp b/Engine/Plugins/VirtualProduction/PropertyAnimatorCore/Source/PropertyAnimatorCore/Private/Properties/Handlers/PropertyAnimatorCoreIntHandler.cpp
--- a/Engine/Plugins/VirtualProduction/PropertyAnimatorCore/Source/PropertyAnimatorCore/Private/Properties/Handlers/PropertyAnimatorCoreIntHandler.cpp
+++ b/Engine/Plugins/VirtualProduction/PropertyAnimatorCore/Source/PropertyAnimatorCore/Private/Properties/Handlers/PropertyAnimatorCoreIntHandler.cpp
@@ -63,6 +63,22 @@ bool UPropertyAnimatorCoreIntHandler::AddValue(const FPropertyAnimatorCoreData&
 	return true;
 }
 
+bool UPropertyAnimatorCoreIntHandler::AddValue(const FPropertyAnimatorCoreData& InPropertyData, const FInstancedPropertyBag& InOriginalValue, const FInstancedPropertyBag& InDeltaValue)
+{
+	const FName PropertyHash(InPropertyData.GetLocatorPathHash());
+	const TValueOrError<int32, EPropertyBagResult> OriginalValueResult = InOriginalValue.GetValueInt32(PropertyHash);
+	const TValueOrError<int32, EPropertyBagResult> DeltaValueResult = InDeltaValue.GetValueInt32(PropertyHash);
+	if (!OriginalValueResult.HasValue() || !DeltaValueResult.HasValue())
+	{
+		return false;
+	}
+
+	int32 NewValue = OriginalValueResult.GetValue() + DeltaValueResult.GetValue();
+	InPropertyData.SetPropertyValuePtr(&NewValue);
+
+	return true;
+}
+
 bool UPropertyAnimatorCoreIntHandler::SubtractValue(const FPropertyAnimatorCoreData& InPropertyData, const FInstancedPropertyBag& InValue)
 {
 	const FName PropertyHash(InPropertyData.GetLocatorPathHash());
diff --git a/Engine/Plugins/VirtualProduction/PropertyAnimatorCore/Source/PropertyAnimatorCore/Private/Properties/Handlers/PropertyAnimatorCoreIntHandler.h b/Engine/Plugins/VirtualProduction/PropertyAnimatorCore/Source/PropertyAnimatorCore/Private/Properties/Handlers/PropertyAnimatorCoreIntHandler.h
--- a/Engine/Plugins/VirtualProduction/PropertyAnimatorCore/Source/PropertyAnimatorCore/Private/Properties/Handlers/PropertyAnimatorCoreIntHandler.h
+++ b/Engine/Plugins/VirtualProduction/PropertyAnimatorCore/Source/PropertyAnimatorCore/Private/Properties/Handlers/PropertyAnimatorCoreIntHandler.h
@@ -18,6 +18,7 @@ public:
 	virtual bool SetValue(const FPropertyAnimatorCoreData& InPropertyData, const FInstancedPropertyBag& InValue) override;
 	virtual bool IsAdditiveSupported() const override;
 	virtual bool AddValue(const FPropertyAnimatorCoreData& InPropertyData, const FInstancedPropertyBag& InValue) override;
+	virtual bool AddValue(const FPropertyAnimatorCoreData& InPropertyData, const FInstancedPropertyBag& InOriginalValue, const FInstancedPropertyBag& InDeltaValue) override;
 	virtual bool SubtractValue(const FPropertyAnimatorCoreData& InPropertyData, const FInstancedPropertyBag& InValue) override;
 	virtual bool GetDefaultValue(const FPropertyAnimatorCoreData& InPropertyData, FInstancedPropertyBag& OutValue) override;
 	//~ End UPropertyAnimatorCoreHandlerBase
diff --git a/Engine/Plugins/VirtualProduction/PropertyAnimatorCore/Source/PropertyAnimatorCore/Private/Properties/Handlers/PropertyAnimatorCoreRotatorHandler.cpp b/Engine/Plugins/VirtualProduction/PropertyAnimatorCore/Source/PropertyAnimatorCore/Private/Properties/Handlers/PropertyAnimatorCoreRotatorHandler.cpp
--- a/Engine/Plugins/VirtualProduction/PropertyAnimatorCore/Source/PropertyAnimatorCore/Private/Properties/Handlers/PropertyAnimatorCoreRotatorHandler.cpp
+++ b/Engine/Plugins/VirtualProduction/PropertyAnimatorCore/Source/PropertyAnimatorCore/Private/Properties/Handlers/PropertyAnimatorCoreRotatorHandler.cpp
@@ -63,6 +63,22 @@ bool UPropertyAnimatorCoreRotatorHandler::AddValue(const FPropertyAnimatorCoreDa
 	return true;
 }
 
+bool UPropertyAnimatorCoreRotatorHandler::AddValue(const FPropertyAnimatorCoreData& InPropertyData, const FInstancedPropertyBag& InOriginalValue, const FInstancedPropertyBag& InDeltaValue)
+{
+	const FName PropertyHash(InPropertyData.GetLocatorPathHash());
+	const TValueOrError<FRotator*, EPropertyBagResult> OriginalValueResult = InOriginalValue.GetValueStruct<FRotator>(PropertyHash);
+	const TValueOrError<FRotator*, EPropertyBagResult> DeltaValueResult = InDeltaValue.GetValueStruct<FRotator>(PropertyHash);
+	if (!OriginalValueResult.HasValue() || !DeltaValueResult.HasValue())
+	{
+		return false;
+	}
+
+	FRotator NewValue = FRotator(OriginalValueResult.GetValue()->Quaternion() * DeltaValueResult.GetValue()->Quaternion());
+	InPropertyData.SetPropertyValuePtr(&NewValue);
+
+	return true;
+}
+
 bool UPropertyAnimatorCoreRotatorHandler::SubtractValue(const FPropertyAnimatorCoreData& InPropertyData, const FInstancedPropertyBag& InValue)
 {
 	const FName PropertyHash(InPropertyData.GetLocatorPathHash());
diff --git a/Engine/Plugins/VirtualProduction/PropertyAnimatorCore/Source/PropertyAnimatorCore/Private/Properties/Handlers/PropertyAnimatorCoreRotatorHandler.h b/Engine/Plugins/VirtualProduction/PropertyAnimatorCore/Source/PropertyAnimatorCore/Private/Properties/Handlers/PropertyAnimatorCoreRotatorHandler.h
--- a/Engine/Plugins/VirtualProduction/PropertyAnimatorCore/Source/PropertyAnimatorCore/Private/Properties/Handlers/PropertyAnimatorCoreRotatorHandler.h
+++ b/Engine/Plugins/VirtualProduction/PropertyAnimatorCore/Source/PropertyAnimatorCore/Private/Properties/Handlers/PropertyAnimatorCoreRotatorHandler.h
@@ -18,6 +18,7 @@ public:
 	virtual bool SetValue(const FPropertyAnimatorCoreData& InPropertyData, const FInstancedPropertyBag& InValue) override;
 	virtual bool IsAdditiveSupported() const override;
 	virtual bool AddValue(const FPropertyAnimatorCoreData& InPropertyData, const FInstancedPropertyBag& InValue) override;
+	virtual bool AddValue(const FPropertyAnimatorCoreData& InPropertyData, const FInstancedPropertyBag& InOriginalValue, const FInstancedPropertyBag& InDeltaValue) override;
 	virtual bool SubtractValue(const FPropertyAnimatorCoreData& InPropertyData, const FInstancedPropertyBag& InValue) override;
 	virtual bool GetDefaultValue(const FPropertyAnimatorCoreData& InPropertyData, FInstancedPropertyBag& OutValue) override;
 	//~ End UPropertyAnimatorCoreHandlerBase
diff --git a/Engine/Plugins/VirtualProduction/PropertyAnimatorCore/Source/PropertyAnimatorCore/Private/Properties/Handlers/PropertyAnimatorCoreStringHandler.cpp b/Engine/Plugins/VirtualProduction/PropertyAnimatorCore/Source/PropertyAnimatorCore/Private/Properties/Handlers/PropertyAnimatorCoreStringHandler.cpp
--- a/Engine/Plugins/VirtualProduction/PropertyAnimatorCore/Source/PropertyAnimatorCore/Private/Properties/Handlers/PropertyAnimatorCoreStringHandler.cpp
+++ b/Engine/Plugins/VirtualProduction/PropertyAnimatorCore/Source/PropertyAnimatorCore/Private/Properties/Handlers/PropertyAnimatorCoreStringHandler.cpp
@@ -64,6 +64,23 @@ bool UPropertyAnimatorCoreStringHandler::AddValue(const FPropertyAnimatorCoreDat
 	return true;
 }
 
+bool UPropertyAnimatorCoreStringHandler::AddValue(const FPropertyAnimatorCoreData& InPropertyData, const FInstancedPropertyBag& InOriginalValue, const FInstancedPropertyBag& InDeltaValue)
+{
+	const FName PropertyHash(InPropertyData.GetLocatorPathHash());
+	const TValueOrError<FString, EPropertyBagResult> OriginalValueResult = InOriginalValue.GetValueString(PropertyHash);
+	const TValueOrError<FString, EPropertyBagResult> DeltaValueResult = InDeltaValue.GetValueString(PropertyHash);
+	if (!OriginalValueResult.HasValue() || !DeltaValueResult.HasValue())
+	{
+		return false;
+	}
+
+	// Append
+	FString NewValue = OriginalValueResult.GetValue() + DeltaValueResult.GetValue();
+	InPropertyData.SetPropertyValuePtr(&NewValue);
+
+	return true;
+}
+
 bool UPropertyAnimatorCoreStringHandler::SubtractValue(const FPropertyAnimatorCoreData& InPropertyData, const FInstancedPropertyBag& InValue)
 {
 	const FName PropertyHash(InPropertyData.GetLocatorPathHash());
diff --git a/Engine/Plugins/VirtualProduction/PropertyAnimatorCore/Source/PropertyAnimatorCore/Private/Properties/Handlers/PropertyAnimatorCoreStringHandler.h b/Engine/Plugins/VirtualProduction/PropertyAnimatorCore/Source/PropertyAnimatorCore/Private/Properties/Handlers/PropertyAnimatorCoreStringHandler.h
--- a/Engine/Plugins/VirtualProduction/PropertyAnimatorCore/Source/PropertyAnimatorCore/Private/Properties/Handlers/PropertyAnimatorCoreStringHandler.h
+++ b/Engine/Plugins/VirtualProduction/PropertyAnimatorCore/Source/PropertyAnimatorCore/Private/Properties/Handlers/PropertyAnimatorCoreStringHandler.h
@@ -18,6 +18,7 @@ public:
 	virtual bool SetValue(const FPropertyAnimatorCoreData& InPropertyData, const FInstancedPropertyBag& InValue) override;
 	virtual bool IsAdditiveSupported() const override;
 	virtual bool AddValue(const FPropertyAnimatorCoreData& InPropertyData, const FInstancedPropertyBag& InValue) override;
+	virtual bool AddValue(const FPropertyAnimatorCoreData& InPropertyData, const FInstancedPropertyBag& InOriginalValue, const FInstancedPropertyBag& InDeltaValue) override;
 	virtual bool SubtractValue(const FPropertyAnimatorCoreData& InPropertyData, const FInstancedPropertyBag& InValue) override;
 	virtual bool GetDefaultValue(const FPropertyAnimatorCoreData& InPropertyData, FInstancedPropertyBag& OutValue) override;
 	//~ End UPropertyAnimatorCoreHandlerBase
diff --git a/Engine/Plugins/VirtualProduction/PropertyAnimatorCore/Source/PropertyAnimatorCore/Private/Properties/Handlers/PropertyAnimatorCoreTextHandler.cpp b/Engine/Plugins/VirtualProduction/PropertyAnimatorCore/Source/PropertyAnimatorCore/Private/Properties/Handlers/PropertyAnimatorCoreTextHandler.cpp
--- a/Engine/Plugins/VirtualProduction/PropertyAnimatorCore/Source/PropertyAnimatorCore/Private/Properties/Handlers/PropertyAnimatorCoreTextHandler.cpp
+++ b/Engine/Plugins/VirtualProduction/PropertyAnimatorCore/Source/PropertyAnimatorCore/Private/Properties/Handlers/PropertyAnimatorCoreTextHandler.cpp
@@ -60,7 +60,24 @@ bool UPropertyAnimatorCoreTextHandler::AddValue(const FPropertyAnimatorCoreData&
 	InPropertyData.GetPropertyValuePtr(&Value);
 
 	// Append
-	FText NewValue = FText::FromString(Value.ToString() + ValueResult.GetValue().ToString());
+	FText NewValue = FText::Join(FText::GetEmpty(), Value, ValueResult.GetValue());
+	InPropertyData.SetPropertyValuePtr(&NewValue);
+
+	return true;
+}
+
+bool UPropertyAnimatorCoreTextHandler::AddValue(const FPropertyAnimatorCoreData& InPropertyData, const FInstancedPropertyBag& InOriginalValue, const FInstancedPropertyBag& InDeltaValue)
+{
+	const FName PropertyHash(InPropertyData.GetLocatorPathHash());
+	const TValueOrError<FText, EPropertyBagResult> OriginalValueResult = InOriginalValue.GetValueText(PropertyHash);
+	const TValueOrError<FText, EPropertyBagResult> DeltaValueResult = InDeltaValue.GetValueText(PropertyHash);
+	if (!OriginalValueResult.HasValue() || !DeltaValueResult.HasValue())
+	{
+		return false;
+	}
+
+	// Append
+	FText NewValue = FText::Join(FText::GetEmpty(), OriginalValueResult.GetValue(), DeltaValueResult.GetValue());
 	InPropertyData.SetPropertyValuePtr(&NewValue);
 
 	return true;
diff --git a/Engine/Plugins/VirtualProduction/PropertyAnimatorCore/Source/PropertyAnimatorCore/Private/Properties/Handlers/PropertyAnimatorCoreTextHandler.h b/Engine/Plugins/VirtualProduction/PropertyAnimatorCore/Source/PropertyAnimatorCore/Private/Properties/Handlers/PropertyAnimatorCoreTextHandler.h
--- a/Engine/Plugins/VirtualProduction/PropertyAnimatorCore/Source/PropertyAnimatorCore/Private/Properties/Handlers/PropertyAnimatorCoreTextHandler.h
+++ b/Engine/Plugins/VirtualProduction/PropertyAnimatorCore/Source/PropertyAnimatorCore/Private/Properties/Handlers/PropertyAnimatorCoreTextHandler.h
@@ -18,6 +18,7 @@ public:
 	virtual bool SetValue(const FPropertyAnimatorCoreData& InPropertyData, const FInstancedPropertyBag& InValue) override;
 	virtual bool IsAdditiveSupported() const override;
 	virtual bool AddValue(const FPropertyAnimatorCoreData& InPropertyData, const FInstancedPropertyBag& InValue) override;
+	virtual bool AddValue(const FPropertyAnimatorCoreData& InPropertyData, const FInstancedPropertyBag& InOriginalValue, const FInstancedPropertyBag& InDeltaValue) override;
 	virtual bool SubtractValue(const FPropertyAnimatorCoreData& InPropertyData, const FInstancedPropertyBag& InValue) override;
 	virtual bool GetDefaultValue(const FPropertyAnimatorCoreData& InPropertyData, FInstancedPropertyBag& OutValue) override;
 	//~ End UPropertyAnimatorCoreHandlerBase
diff --git a/Engine/Plugins/VirtualProduction/PropertyAnimatorCore/Source/PropertyAnimatorCore/Private/Properties/Handlers/PropertyAnimatorCoreVectorHandler.cpp b/Engine/Plugins/VirtualProduction/PropertyAnimatorCore/Source/PropertyAnimatorCore/Private/Properties/Handlers/PropertyAnimatorCoreVectorHandler.cpp
--- a/Engine/Plugins/VirtualProduction/PropertyAnimatorCore/Source/PropertyAnimatorCore/Private/Properties/Handlers/PropertyAnimatorCoreVectorHandler.cpp
+++ b/Engine/Plugins/VirtualProduction/PropertyAnimatorCore/Source/PropertyAnimatorCore/Private/Properties/Handlers/PropertyAnimatorCoreVectorHandler.cpp
@@ -63,6 +63,22 @@ bool UPropertyAnimatorCoreVectorHandler::AddValue(const FPropertyAnimatorCoreDat
 	return true;
 }
 
+bool UPropertyAnimatorCoreVectorHandler::AddValue(const FPropertyAnimatorCoreData& InPropertyData, const FInstancedPropertyBag& InOriginalValue, const FInstancedPropertyBag& InDeltaValue)
+{
+	const FName PropertyHash(InPropertyData.GetLocatorPathHash());
+	const TValueOrError<FVector*, EPropertyBagResult> OriginalValueResult = InOriginalValue.GetValueStruct<FVector>(PropertyHash);
+	const TValueOrError<FVector*, EPropertyBagResult> DeltaValueResult = InOriginalValue.GetValueStruct<FVector>(PropertyHash);
+	if (!OriginalValueResult.HasValue() || !DeltaValueResult.HasValue())
+	{
+		return false;
+	}
+
+	FVector NewValue = *OriginalValueResult.GetValue() + *DeltaValueResult.GetValue();
+	InPropertyData.SetPropertyValuePtr(&NewValue);
+
+	return true;
+}
+
 bool UPropertyAnimatorCoreVectorHandler::SubtractValue(const FPropertyAnimatorCoreData& InPropertyData, const FInstancedPropertyBag& InValue)
 {
 	const FName PropertyHash(InPropertyData.GetLocatorPathHash());
diff --git a/Engine/Plugins/VirtualProduction/PropertyAnimatorCore/Source/PropertyAnimatorCore/Private/Properties/Handlers/PropertyAnimatorCoreVectorHandler.h b/Engine/Plugins/VirtualProduction/PropertyAnimatorCore/Source/PropertyAnimatorCore/Private/Properties/Handlers/PropertyAnimatorCoreVectorHandler.h
--- a/Engine/Plugins/VirtualProduction/PropertyAnimatorCore/Source/PropertyAnimatorCore/Private/Properties/Handlers/PropertyAnimatorCoreVectorHandler.h
+++ b/Engine/Plugins/VirtualProduction/PropertyAnimatorCore/Source/PropertyAnimatorCore/Private/Properties/Handlers/PropertyAnimatorCoreVectorHandler.h
@@ -18,6 +18,7 @@ public:
 	virtual bool SetValue(const FPropertyAnimatorCoreData& InPropertyData, const FInstancedPropertyBag& InValue) override;
 	virtual bool IsAdditiveSupported() const override;
 	virtual bool AddValue(const FPropertyAnimatorCoreData& InPropertyData, const FInstancedPropertyBag& InValue) override;
+	virtual bool AddValue(const FPropertyAnimatorCoreData& InPropertyData, const FInstancedPropertyBag& InOriginalValue, const FInstancedPropertyBag& InDeltaValue) override;
 	virtual bool SubtractValue(const FPropertyAnimatorCoreData& InPropertyData, const FInstancedPropertyBag& InValue) override;
 	virtual bool GetDefaultValue(const FPropertyAnimatorCoreData& InPropertyData, FInstancedPropertyBag& OutValue) override;
 	//~ End UPropertyAnimatorCoreHandlerBase
diff --git a/Engine/Plugins/VirtualProduction/PropertyAnimatorCore/Source/PropertyAnimatorCore/Private/Properties/PropertyAnimatorCoreContext.cpp b/Engine/Plugins/VirtualProduction/PropertyAnimatorCore/Source/PropertyAnimatorCore/Private/Properties/PropertyAnimatorCoreContext.cpp
--- a/Engine/Plugins/VirtualProduction/PropertyAnimatorCore/Source/PropertyAnimatorCore/Private/Properties/PropertyAnimatorCoreContext.cpp
+++ b/Engine/Plugins/VirtualProduction/PropertyAnimatorCore/Source/PropertyAnimatorCore/Private/Properties/PropertyAnimatorCoreContext.cpp
@@ -4,7 +4,6 @@
 
 #include "Containers/Ticker.h"
 #include "Logs/PropertyAnimatorCoreLog.h"
-#include "Misc/EnumerateRange.h"
 #include "Properties/PropertyAnimatorCoreResolver.h"
 #include "Properties/Converters/PropertyAnimatorCoreConverterBase.h"
 #include "Properties/Handlers/PropertyAnimatorCoreHandlerBase.h"
@@ -104,7 +103,7 @@ void UPropertyAnimatorCoreContext::SetMode(EPropertyAnimatorCoreMode InMode)
 	OnModeChanged();
 }
 
-void UPropertyAnimatorCoreContext::SetConverterClass(TSubclassOf<UPropertyAnimatorCoreConverterBase> InConverterClass)
+void UPropertyAnimatorCoreContext::SetConverterClass(const TSubclassOf<UPropertyAnimatorCoreConverterBase>& InConverterClass)
 {
 	ConverterClass = InConverterClass;
 
@@ -427,8 +426,6 @@ void UPropertyAnimatorCoreContext::Restore()
 			// Set original value
 			Handler->SetValue(ResolvedProperty, OriginalPropertyValues);
 		}
-
-		OriginalPropertyValues.Reset();
 	}
 	else
 	{
@@ -439,6 +436,7 @@ void UPropertyAnimatorCoreContext::Restore()
 		}
 	}
 
+	OriginalPropertyValues.Reset();
 	DeltaPropertyValues.Reset();
 }
 
@@ -525,6 +523,6 @@ void UPropertyAnimatorCoreContext::CommitEvaluationResult(const FPropertyAnimato
 	}
 	else
 	{
-		Handler->AddValue(InResolvedProperty, DeltaPropertyValues);
+		Handler->AddValue(InResolvedProperty, OriginalPropertyValues, DeltaPropertyValues);
 	}
 }
diff --git a/Engine/Plugins/VirtualProduction/PropertyAnimatorCore/Source/PropertyAnimatorCore/Public/Animators/PropertyAnimatorCoreBase.h b/Engine/Plugins/VirtualProduction/PropertyAnimatorCore/Source/PropertyAnimatorCore/Public/Animators/PropertyAnimatorCoreBase.h
--- a/Engine/Plugins/VirtualProduction/PropertyAnimatorCore/Source/PropertyAnimatorCore/Public/Animators/PropertyAnimatorCoreBase.h
+++ b/Engine/Plugins/VirtualProduction/PropertyAnimatorCore/Source/PropertyAnimatorCore/Public/Animators/PropertyAnimatorCoreBase.h
@@ -326,7 +326,7 @@ private:
 	PROPERTYANIMATORCORE_API static FOnAnimatorPropertyUpdated OnAnimatorPropertyUnlinkedDelegate;
 
 	/** Restore modified properties to original state */
-	void RestoreProperties(bool bInForce = false);
+	void RestoreProperties();
 
 	/** Allocate and saves properties in the property bag */
 	void SaveProperties();
diff --git a/Engine/Plugins/VirtualProduction/PropertyAnimatorCore/Source/PropertyAnimatorCore/Public/Properties/Handlers/PropertyAnimatorCoreHandlerBase.h b/Engine/Plugins/VirtualProduction/PropertyAnimatorCore/Source/PropertyAnimatorCore/Public/Properties/Handlers/PropertyAnimatorCoreHandlerBase.h
--- a/Engine/Plugins/VirtualProduction/PropertyAnimatorCore/Source/PropertyAnimatorCore/Public/Properties/Handlers/PropertyAnimatorCoreHandlerBase.h
+++ b/Engine/Plugins/VirtualProduction/PropertyAnimatorCore/Source/PropertyAnimatorCore/Public/Properties/Handlers/PropertyAnimatorCoreHandlerBase.h
@@ -54,6 +54,12 @@ public:
 		return false;
 	}
 
+	/** Adds the value to this property within the owner, delta value and original value are in property bags */
+	virtual bool AddValue(const FPropertyAnimatorCoreData& InPropertyData, const FInstancedPropertyBag& InOriginalValue, const FInstancedPropertyBag& InDeltaValue)
+	{
+		return false;
+	}
+
 	/** Get the default value for this property, init to zero */
 	virtual bool GetDefaultValue(const FPropertyAnimatorCoreData& InPropertyData, FInstancedPropertyBag& OutValue)
 	{
diff --git a/Engine/Plugins/VirtualProduction/PropertyAnimatorCore/Source/PropertyAnimatorCore/Public/Properties/PropertyAnimatorCoreContext.h b/Engine/Plugins/VirtualProduction/PropertyAnimatorCore/Source/PropertyAnimatorCore/Public/Properties/PropertyAnimatorCoreContext.h
--- a/Engine/Plugins/VirtualProduction/PropertyAnimatorCore/Source/PropertyAnimatorCore/Public/Properties/PropertyAnimatorCoreContext.h
+++ b/Engine/Plugins/VirtualProduction/PropertyAnimatorCore/Source/PropertyAnimatorCore/Public/Properties/PropertyAnimatorCoreContext.h
@@ -75,7 +75,7 @@ public:
 		return Mode;
 	}
 
-	PROPERTYANIMATORCORE_API void SetConverterClass(TSubclassOf<UPropertyAnimatorCoreConverterBase> InConverterClass);
+	PROPERTYANIMATORCORE_API void SetConverterClass(const TSubclassOf<UPropertyAnimatorCoreConverterBase>& InConverterClass);
 	TSubclassOf<UPropertyAnimatorCoreConverterBase> GetConverterClass() const
 	{
 		return ConverterClass;
