MotionDesign : PropertyAnimator
- Fixed additive mode not working when value is changed while animator is active

#DisabledValidations NonUnityAndStaticAnalysis

diff --git a/Engine/Plugins/VirtualProduction/PropertyAnimatorCore/Source/PropertyAnimatorCore/Private/Properties/Converters/PropertyAnimatorCoreFloatToBoolConverter.cpp b/Engine/Plugins/VirtualProduction/PropertyAnimatorCore/Source/PropertyAnimatorCore/Private/Properties/Converters/PropertyAnimatorCoreFloatToBoolConverter.cpp
--- a/Engine/Plugins/VirtualProduction/PropertyAnimatorCore/Source/PropertyAnimatorCore/Private/Properties/Converters/PropertyAnimatorCoreFloatToBoolConverter.cpp
+++ b/Engine/Plugins/VirtualProduction/PropertyAnimatorCore/Source/PropertyAnimatorCore/Private/Properties/Converters/PropertyAnimatorCoreFloatToBoolConverter.cpp
@@ -27,6 +27,7 @@ bool UPropertyAnimatorCoreFloatToBoolConverter::Convert(const FPropertyBagProper
 	bool bOutValue;
 	if (TValueConverterTraits<float, bool, FBoolConverterRule>::Convert(FromValueResult.GetValue(), bOutValue, InRule.GetPtr<FBoolConverterRule>()))
 	{
+		InToBag.AddProperty(InToProperty.Name, EPropertyBagPropertyType::Bool);
 		InToBag.SetValueBool(InToProperty.Name, bOutValue);
 		return true;
 	}
diff --git a/Engine/Plugins/VirtualProduction/PropertyAnimatorCore/Source/PropertyAnimatorCore/Private/Properties/Converters/PropertyAnimatorCoreFloatToDoubleConverter.cpp b/Engine/Plugins/VirtualProduction/PropertyAnimatorCore/Source/PropertyAnimatorCore/Private/Properties/Converters/PropertyAnimatorCoreFloatToDoubleConverter.cpp
--- a/Engine/Plugins/VirtualProduction/PropertyAnimatorCore/Source/PropertyAnimatorCore/Private/Properties/Converters/PropertyAnimatorCoreFloatToDoubleConverter.cpp
+++ b/Engine/Plugins/VirtualProduction/PropertyAnimatorCore/Source/PropertyAnimatorCore/Private/Properties/Converters/PropertyAnimatorCoreFloatToDoubleConverter.cpp
@@ -27,6 +27,7 @@ bool UPropertyAnimatorCoreFloatToDoubleConverter::Convert(const FPropertyBagProp
 	double OutValue;
 	if (TValueConverterTraits<float, double>::Convert(FromValueResult.GetValue(), OutValue, nullptr))
 	{
+		InToBag.AddProperty(InToProperty.Name, EPropertyBagPropertyType::Double);
 		InToBag.SetValueDouble(InToProperty.Name, OutValue);
 		return true;
 	}
diff --git a/Engine/Plugins/VirtualProduction/PropertyAnimatorCore/Source/PropertyAnimatorCore/Private/Properties/Converters/PropertyAnimatorCoreFloatToIntConverter.cpp b/Engine/Plugins/VirtualProduction/PropertyAnimatorCore/Source/PropertyAnimatorCore/Private/Properties/Converters/PropertyAnimatorCoreFloatToIntConverter.cpp
--- a/Engine/Plugins/VirtualProduction/PropertyAnimatorCore/Source/PropertyAnimatorCore/Private/Properties/Converters/PropertyAnimatorCoreFloatToIntConverter.cpp
+++ b/Engine/Plugins/VirtualProduction/PropertyAnimatorCore/Source/PropertyAnimatorCore/Private/Properties/Converters/PropertyAnimatorCoreFloatToIntConverter.cpp
@@ -28,6 +28,7 @@ bool UPropertyAnimatorCoreFloatToIntConverter::Convert(const FPropertyBagPropert
 	int32 OutValue;
 	if (TValueConverterTraits<float, int32, FInt32ConverterRule>::Convert(FromValueResult.GetValue(), OutValue, InRule.GetPtr<FInt32ConverterRule>()))
 	{
+		InToBag.AddProperty(InToProperty.Name, EPropertyBagPropertyType::Int32);
 		InToBag.SetValueDouble(InToProperty.Name, OutValue);
 		return true;
 	}
diff --git a/Engine/Plugins/VirtualProduction/PropertyAnimatorCore/Source/PropertyAnimatorCore/Private/Properties/Converters/PropertyAnimatorCoreStringToTextConverter.cpp b/Engine/Plugins/VirtualProduction/PropertyAnimatorCore/Source/PropertyAnimatorCore/Private/Properties/Converters/PropertyAnimatorCoreStringToTextConverter.cpp
--- a/Engine/Plugins/VirtualProduction/PropertyAnimatorCore/Source/PropertyAnimatorCore/Private/Properties/Converters/PropertyAnimatorCoreStringToTextConverter.cpp
+++ b/Engine/Plugins/VirtualProduction/PropertyAnimatorCore/Source/PropertyAnimatorCore/Private/Properties/Converters/PropertyAnimatorCoreStringToTextConverter.cpp
@@ -27,6 +27,7 @@ bool UPropertyAnimatorCoreStringToTextConverter::Convert(const FPropertyBagPrope
 	FText OutValue;
 	if (TValueConverterTraits<FString, FText, void>::Convert(FromValueResult.GetValue(), OutValue, nullptr))
 	{
+		InToBag.AddProperty(InToProperty.Name, EPropertyBagPropertyType::Text);
 		InToBag.SetValueText(InToProperty.Name, OutValue);
 		return true;
 	}
diff --git a/Engine/Plugins/VirtualProduction/PropertyAnimatorCore/Source/PropertyAnimatorCore/Private/Properties/Converters/PropertyAnimatorCoreTextToStringConverter.cpp b/Engine/Plugins/VirtualProduction/PropertyAnimatorCore/Source/PropertyAnimatorCore/Private/Properties/Converters/PropertyAnimatorCoreTextToStringConverter.cpp
--- a/Engine/Plugins/VirtualProduction/PropertyAnimatorCore/Source/PropertyAnimatorCore/Private/Properties/Converters/PropertyAnimatorCoreTextToStringConverter.cpp
+++ b/Engine/Plugins/VirtualProduction/PropertyAnimatorCore/Source/PropertyAnimatorCore/Private/Properties/Converters/PropertyAnimatorCoreTextToStringConverter.cpp
@@ -27,6 +27,7 @@ bool UPropertyAnimatorCoreTextToStringConverter::Convert(const FPropertyBagPrope
 	FString OutValue;
 	if (TValueConverterTraits<FText, FString>::Convert(FromValueResult.GetValue(), OutValue, nullptr))
 	{
+		InToBag.AddProperty(InToProperty.Name, EPropertyBagPropertyType::String);
 		InToBag.SetValueString(InToProperty.Name, OutValue);
 		return true;
 	}
diff --git a/Engine/Plugins/VirtualProduction/PropertyAnimatorCore/Source/PropertyAnimatorCore/Private/Properties/Handlers/PropertyAnimatorCoreDoubleHandler.cpp b/Engine/Plugins/VirtualProduction/PropertyAnimatorCore/Source/PropertyAnimatorCore/Private/Properties/Handlers/PropertyAnimatorCoreDoubleHandler.cpp
--- a/Engine/Plugins/VirtualProduction/PropertyAnimatorCore/Source/PropertyAnimatorCore/Private/Properties/Handlers/PropertyAnimatorCoreDoubleHandler.cpp
+++ b/Engine/Plugins/VirtualProduction/PropertyAnimatorCore/Source/PropertyAnimatorCore/Private/Properties/Handlers/PropertyAnimatorCoreDoubleHandler.cpp
@@ -68,20 +68,54 @@ bool UPropertyAnimatorCoreDoubleHandler::AddValue(const FPropertyAnimatorCoreDat
 	return true;
 }
 
-bool UPropertyAnimatorCoreDoubleHandler::AddValue(const FPropertyAnimatorCoreData& InPropertyData, const FInstancedPropertyBag& InOriginalValue, const FInstancedPropertyBag& InDeltaValue)
+bool UPropertyAnimatorCoreDoubleHandler::AddValue(const FPropertyAnimatorCoreData& InPropertyData, const FInstancedPropertyBag& InValueA, const FInstancedPropertyBag& InValueB, FInstancedPropertyBag& OutValue)
 {
 	const FName PropertyHash(InPropertyData.GetLocatorPathHash());
-	const TValueOrError<double, EPropertyBagResult> OriginalValueResult = InOriginalValue.GetValueDouble(PropertyHash);
-	const TValueOrError<double, EPropertyBagResult> DeltaValueResult = InDeltaValue.GetValueDouble(PropertyHash);
-	if (!OriginalValueResult.HasValue() || !DeltaValueResult.HasValue())
+	const TValueOrError<double, EPropertyBagResult> ValueAResult = InValueA.GetValueDouble(PropertyHash);
+	const TValueOrError<double, EPropertyBagResult> ValueBResult = InValueB.GetValueDouble(PropertyHash);
+	if (!ValueAResult.HasValue() || !ValueBResult.HasValue())
 	{
 		return false;
 	}
 
-	double NewValue = OriginalValueResult.GetValue() + DeltaValueResult.GetValue();
-	InPropertyData.SetPropertyValuePtr(&NewValue);
+	const double NewValue = ValueAResult.GetValue() + ValueBResult.GetValue();
+	const EPropertyBagResult Result = OutValue.SetValueDouble(PropertyHash, NewValue);
 
-	return true;
+	return Result == EPropertyBagResult::Success;
+}
+
+bool UPropertyAnimatorCoreDoubleHandler::AddValue(const FPropertyAnimatorCoreData& InPropertyData, const FInstancedPropertyBag& InValue, FInstancedPropertyBag& OutValue)
+{
+	const FName PropertyHash(InPropertyData.GetLocatorPathHash());
+	const TValueOrError<double, EPropertyBagResult> ValueResult = InValue.GetValueDouble(PropertyHash);
+	if (!ValueResult.HasValue())
+	{
+		return false;
+	}
+
+	double Value;
+	InPropertyData.GetPropertyValuePtr(&Value);
+
+	double NewValue = Value + ValueResult.GetValue();
+	const EPropertyBagResult Result = OutValue.SetValueDouble(PropertyHash, NewValue);
+
+	return Result == EPropertyBagResult::Success;
+}
+
+bool UPropertyAnimatorCoreDoubleHandler::SubtractValue(const FPropertyAnimatorCoreData& InPropertyData, const FInstancedPropertyBag& InValueA, const FInstancedPropertyBag& InValueB, FInstancedPropertyBag& OutValue)
+{
+	const FName PropertyHash(InPropertyData.GetLocatorPathHash());
+	const TValueOrError<double, EPropertyBagResult> ValueAResult = InValueA.GetValueDouble(PropertyHash);
+	const TValueOrError<double, EPropertyBagResult> ValueBResult = InValueB.GetValueDouble(PropertyHash);
+	if (!ValueAResult.HasValue() || !ValueBResult.HasValue())
+	{
+		return false;
+	}
+
+	const double NewValue = ValueAResult.GetValue() - ValueBResult.GetValue();
+	const EPropertyBagResult Result = OutValue.SetValueDouble(PropertyHash, NewValue);
+
+	return Result == EPropertyBagResult::Success;
 }
 
 bool UPropertyAnimatorCoreDoubleHandler::SubtractValue(const FPropertyAnimatorCoreData& InPropertyData, const FInstancedPropertyBag& InValue)
@@ -102,6 +136,24 @@ bool UPropertyAnimatorCoreDoubleHandler::SubtractValue(const FPropertyAnimatorCo
 	return true;
 }
 
+bool UPropertyAnimatorCoreDoubleHandler::SubtractValue(const FPropertyAnimatorCoreData& InPropertyData, const FInstancedPropertyBag& InValue, FInstancedPropertyBag& OutValue)
+{
+	const FName PropertyHash(InPropertyData.GetLocatorPathHash());
+	const TValueOrError<double, EPropertyBagResult> ValueResult = InValue.GetValueDouble(PropertyHash);
+	if (!ValueResult.HasValue())
+	{
+		return false;
+	}
+
+	double Value;
+	InPropertyData.GetPropertyValuePtr(&Value);
+
+	const double NewValue = Value - ValueResult.GetValue();
+	const EPropertyBagResult Result = OutValue.SetValueDouble(PropertyHash, NewValue);
+
+	return Result == EPropertyBagResult::Success;
+}
+
 bool UPropertyAnimatorCoreDoubleHandler::GetDefaultValue(const FPropertyAnimatorCoreData& InPropertyData, FInstancedPropertyBag& OutValue)
 {
 	const FName PropertyHash(InPropertyData.GetLocatorPathHash());
diff --git a/Engine/Plugins/VirtualProduction/PropertyAnimatorCore/Source/PropertyAnimatorCore/Private/Properties/Handlers/PropertyAnimatorCoreDoubleHandler.h b/Engine/Plugins/VirtualProduction/PropertyAnimatorCore/Source/PropertyAnimatorCore/Private/Properties/Handlers/PropertyAnimatorCoreDoubleHandler.h
--- a/Engine/Plugins/VirtualProduction/PropertyAnimatorCore/Source/PropertyAnimatorCore/Private/Properties/Handlers/PropertyAnimatorCoreDoubleHandler.h
+++ b/Engine/Plugins/VirtualProduction/PropertyAnimatorCore/Source/PropertyAnimatorCore/Private/Properties/Handlers/PropertyAnimatorCoreDoubleHandler.h
@@ -18,8 +18,11 @@ public:
 	virtual bool SetValue(const FPropertyAnimatorCoreData& InPropertyData, const FInstancedPropertyBag& InValue) override;
 	virtual bool IsAdditiveSupported() const override;
 	virtual bool AddValue(const FPropertyAnimatorCoreData& InPropertyData, const FInstancedPropertyBag& InValue) override;
-	virtual bool AddValue(const FPropertyAnimatorCoreData& InPropertyData, const FInstancedPropertyBag& InOriginalValue, const FInstancedPropertyBag& InDeltaValue) override;
+	virtual bool AddValue(const FPropertyAnimatorCoreData& InPropertyData, const FInstancedPropertyBag& InValueA, const FInstancedPropertyBag& InValueB, FInstancedPropertyBag& OutValue) override;
+	virtual bool AddValue(const FPropertyAnimatorCoreData& InPropertyData, const FInstancedPropertyBag& InValue, FInstancedPropertyBag& OutValue) override;
+	virtual bool SubtractValue(const FPropertyAnimatorCoreData& InPropertyData, const FInstancedPropertyBag& InValueA, const FInstancedPropertyBag& InValueB, FInstancedPropertyBag& OutValue) override;
 	virtual bool SubtractValue(const FPropertyAnimatorCoreData& InPropertyData, const FInstancedPropertyBag& InValue) override;
+	virtual bool SubtractValue(const FPropertyAnimatorCoreData& InPropertyData, const FInstancedPropertyBag& InValue, FInstancedPropertyBag& OutValue) override;
 	virtual bool GetDefaultValue(const FPropertyAnimatorCoreData& InPropertyData, FInstancedPropertyBag& OutValue) override;
 	//~ End UPropertyAnimatorCoreHandlerBase
 };
diff --git a/Engine/Plugins/VirtualProduction/PropertyAnimatorCore/Source/PropertyAnimatorCore/Private/Properties/Handlers/PropertyAnimatorCoreFloatHandler.cpp b/Engine/Plugins/VirtualProduction/PropertyAnimatorCore/Source/PropertyAnimatorCore/Private/Properties/Handlers/PropertyAnimatorCoreFloatHandler.cpp
--- a/Engine/Plugins/VirtualProduction/PropertyAnimatorCore/Source/PropertyAnimatorCore/Private/Properties/Handlers/PropertyAnimatorCoreFloatHandler.cpp
+++ b/Engine/Plugins/VirtualProduction/PropertyAnimatorCore/Source/PropertyAnimatorCore/Private/Properties/Handlers/PropertyAnimatorCoreFloatHandler.cpp
@@ -68,20 +68,54 @@ bool UPropertyAnimatorCoreFloatHandler::AddValue(const FPropertyAnimatorCoreData
 	return true;
 }
 
-bool UPropertyAnimatorCoreFloatHandler::AddValue(const FPropertyAnimatorCoreData& InPropertyData, const FInstancedPropertyBag& InOriginalValue, const FInstancedPropertyBag& InDeltaValue)
+bool UPropertyAnimatorCoreFloatHandler::AddValue(const FPropertyAnimatorCoreData& InPropertyData, const FInstancedPropertyBag& InValueA, const FInstancedPropertyBag& InValueB, FInstancedPropertyBag& OutValue)
 {
 	const FName PropertyHash(InPropertyData.GetLocatorPathHash());
-	const TValueOrError<float, EPropertyBagResult> OriginalValueResult = InOriginalValue.GetValueFloat(PropertyHash);
-	const TValueOrError<float, EPropertyBagResult> DeltaValueResult = InDeltaValue.GetValueFloat(PropertyHash);
-	if (!OriginalValueResult.HasValue() || !DeltaValueResult.HasValue())
+	const TValueOrError<float, EPropertyBagResult> ValueAResult = InValueA.GetValueFloat(PropertyHash);
+	const TValueOrError<float, EPropertyBagResult> ValueBResult = InValueB.GetValueFloat(PropertyHash);
+	if (!ValueAResult.HasValue() || !ValueBResult.HasValue())
 	{
 		return false;
 	}
 
-	float NewValue = OriginalValueResult.GetValue() + DeltaValueResult.GetValue();
-	InPropertyData.SetPropertyValuePtr(&NewValue);
+	const float NewValue = ValueAResult.GetValue() + ValueBResult.GetValue();
+	EPropertyBagResult Result = OutValue.SetValueFloat(PropertyHash, NewValue);
 
-	return true;
+	return Result == EPropertyBagResult::Success;
+}
+
+bool UPropertyAnimatorCoreFloatHandler::AddValue(const FPropertyAnimatorCoreData& InPropertyData, const FInstancedPropertyBag& InValue, FInstancedPropertyBag& OutValue)
+{
+	const FName PropertyHash(InPropertyData.GetLocatorPathHash());
+	const TValueOrError<float, EPropertyBagResult> ValueResult = InValue.GetValueFloat(PropertyHash);
+	if (!ValueResult.HasValue())
+	{
+		return false;
+	}
+
+	float Value;
+	InPropertyData.GetPropertyValuePtr(&Value);
+
+	float NewValue = Value + ValueResult.GetValue();
+	const EPropertyBagResult Result = OutValue.SetValueFloat(PropertyHash, NewValue);
+
+	return Result == EPropertyBagResult::Success;
+}
+
+bool UPropertyAnimatorCoreFloatHandler::SubtractValue(const FPropertyAnimatorCoreData& InPropertyData, const FInstancedPropertyBag& InValueA, const FInstancedPropertyBag& InValueB, FInstancedPropertyBag& OutValue)
+{
+	const FName PropertyHash(InPropertyData.GetLocatorPathHash());
+	const TValueOrError<float, EPropertyBagResult> ValueAResult = InValueA.GetValueFloat(PropertyHash);
+	const TValueOrError<float, EPropertyBagResult> ValueBResult = InValueB.GetValueFloat(PropertyHash);
+	if (!ValueAResult.HasValue() || !ValueBResult.HasValue())
+	{
+		return false;
+	}
+
+	const float NewValue = ValueAResult.GetValue() - ValueBResult.GetValue();
+	EPropertyBagResult Result = OutValue.SetValueFloat(PropertyHash, NewValue);
+
+	return Result == EPropertyBagResult::Success;
 }
 
 bool UPropertyAnimatorCoreFloatHandler::SubtractValue(const FPropertyAnimatorCoreData& InPropertyData, const FInstancedPropertyBag& InValue)
@@ -102,6 +136,24 @@ bool UPropertyAnimatorCoreFloatHandler::SubtractValue(const FPropertyAnimatorCor
 	return true;
 }
 
+bool UPropertyAnimatorCoreFloatHandler::SubtractValue(const FPropertyAnimatorCoreData& InPropertyData, const FInstancedPropertyBag& InValue, FInstancedPropertyBag& OutValue)
+{
+	const FName PropertyHash(InPropertyData.GetLocatorPathHash());
+	const TValueOrError<float, EPropertyBagResult> ValueResult = InValue.GetValueFloat(PropertyHash);
+	if (!ValueResult.HasValue())
+	{
+		return false;
+	}
+
+	float Value;
+	InPropertyData.GetPropertyValuePtr(&Value);
+
+	const float NewValue = Value - ValueResult.GetValue();
+	const EPropertyBagResult Result = OutValue.SetValueFloat(PropertyHash, NewValue);
+
+	return Result == EPropertyBagResult::Success;
+}
+
 bool UPropertyAnimatorCoreFloatHandler::GetDefaultValue(const FPropertyAnimatorCoreData& InPropertyData, FInstancedPropertyBag& OutValue)
 {
 	const FName PropertyHash(InPropertyData.GetLocatorPathHash());
diff --git a/Engine/Plugins/VirtualProduction/PropertyAnimatorCore/Source/PropertyAnimatorCore/Private/Properties/Handlers/PropertyAnimatorCoreFloatHandler.h b/Engine/Plugins/VirtualProduction/PropertyAnimatorCore/Source/PropertyAnimatorCore/Private/Properties/Handlers/PropertyAnimatorCoreFloatHandler.h
--- a/Engine/Plugins/VirtualProduction/PropertyAnimatorCore/Source/PropertyAnimatorCore/Private/Properties/Handlers/PropertyAnimatorCoreFloatHandler.h
+++ b/Engine/Plugins/VirtualProduction/PropertyAnimatorCore/Source/PropertyAnimatorCore/Private/Properties/Handlers/PropertyAnimatorCoreFloatHandler.h
@@ -18,8 +18,11 @@ public:
 	virtual bool SetValue(const FPropertyAnimatorCoreData& InPropertyData, const FInstancedPropertyBag& InValue) override;
 	virtual bool IsAdditiveSupported() const override;
 	virtual bool AddValue(const FPropertyAnimatorCoreData& InPropertyData, const FInstancedPropertyBag& InValue) override;
-	virtual bool AddValue(const FPropertyAnimatorCoreData& InPropertyData, const FInstancedPropertyBag& InOriginalValue, const FInstancedPropertyBag& InDeltaValue) override;
+	virtual bool AddValue(const FPropertyAnimatorCoreData& InPropertyData, const FInstancedPropertyBag& InValueA, const FInstancedPropertyBag&InValueB, FInstancedPropertyBag& OutValue) override;
+	virtual bool AddValue(const FPropertyAnimatorCoreData& InPropertyData, const FInstancedPropertyBag& InValue, FInstancedPropertyBag& OutValue) override;
+	virtual bool SubtractValue(const FPropertyAnimatorCoreData& InPropertyData, const FInstancedPropertyBag& InValueA, const FInstancedPropertyBag& InValueB, FInstancedPropertyBag& OutValue) override;
 	virtual bool SubtractValue(const FPropertyAnimatorCoreData& InPropertyData, const FInstancedPropertyBag& InValue) override;
+	virtual bool SubtractValue(const FPropertyAnimatorCoreData& InPropertyData, const FInstancedPropertyBag& InValue, FInstancedPropertyBag& OutValue) override;
 	virtual bool GetDefaultValue(const FPropertyAnimatorCoreData& InPropertyData, FInstancedPropertyBag& OutValue) override;
 	//~ End UPropertyAnimatorCoreHandlerBase
 };
\ No newline at end of file
diff --git a/Engine/Plugins/VirtualProduction/PropertyAnimatorCore/Source/PropertyAnimatorCore/Private/Properties/Handlers/PropertyAnimatorCoreIntHandler.cpp b/Engine/Plugins/VirtualProduction/PropertyAnimatorCore/Source/PropertyAnimatorCore/Private/Properties/Handlers/PropertyAnimatorCoreIntHandler.cpp
--- a/Engine/Plugins/VirtualProduction/PropertyAnimatorCore/Source/PropertyAnimatorCore/Private/Properties/Handlers/PropertyAnimatorCoreIntHandler.cpp
+++ b/Engine/Plugins/VirtualProduction/PropertyAnimatorCore/Source/PropertyAnimatorCore/Private/Properties/Handlers/PropertyAnimatorCoreIntHandler.cpp
@@ -63,20 +63,54 @@ bool UPropertyAnimatorCoreIntHandler::AddValue(const FPropertyAnimatorCoreData&
 	return true;
 }
 
-bool UPropertyAnimatorCoreIntHandler::AddValue(const FPropertyAnimatorCoreData& InPropertyData, const FInstancedPropertyBag& InOriginalValue, const FInstancedPropertyBag& InDeltaValue)
+bool UPropertyAnimatorCoreIntHandler::AddValue(const FPropertyAnimatorCoreData& InPropertyData, const FInstancedPropertyBag& InValueA, const FInstancedPropertyBag&InValueB, FInstancedPropertyBag& OutValue)
 {
 	const FName PropertyHash(InPropertyData.GetLocatorPathHash());
-	const TValueOrError<int32, EPropertyBagResult> OriginalValueResult = InOriginalValue.GetValueInt32(PropertyHash);
-	const TValueOrError<int32, EPropertyBagResult> DeltaValueResult = InDeltaValue.GetValueInt32(PropertyHash);
-	if (!OriginalValueResult.HasValue() || !DeltaValueResult.HasValue())
+	const TValueOrError<int32, EPropertyBagResult> ValueAResult = InValueA.GetValueInt32(PropertyHash);
+	const TValueOrError<int32, EPropertyBagResult> ValueBResult = InValueB.GetValueInt32(PropertyHash);
+	if (!ValueAResult.HasValue() || !ValueBResult.HasValue())
 	{
 		return false;
 	}
 
-	int32 NewValue = OriginalValueResult.GetValue() + DeltaValueResult.GetValue();
-	InPropertyData.SetPropertyValuePtr(&NewValue);
+	const int32 NewValue = ValueAResult.GetValue() + ValueBResult.GetValue();
+	EPropertyBagResult Result = OutValue.SetValueInt32(PropertyHash, NewValue);
 
-	return true;
+	return Result == EPropertyBagResult::Success;
+}
+
+bool UPropertyAnimatorCoreIntHandler::AddValue(const FPropertyAnimatorCoreData& InPropertyData, const FInstancedPropertyBag& InValue, FInstancedPropertyBag& OutValue)
+{
+	const FName PropertyHash(InPropertyData.GetLocatorPathHash());
+	const TValueOrError<int32, EPropertyBagResult> ValueResult = InValue.GetValueInt32(PropertyHash);
+	if (!ValueResult.HasValue())
+	{
+		return false;
+	}
+
+	int32 Value;
+	InPropertyData.GetPropertyValuePtr(&Value);
+
+	int32 NewValue = Value + ValueResult.GetValue();
+	const EPropertyBagResult Result = OutValue.SetValueInt32(PropertyHash, NewValue);
+
+	return Result == EPropertyBagResult::Success;
+}
+
+bool UPropertyAnimatorCoreIntHandler::SubtractValue(const FPropertyAnimatorCoreData& InPropertyData, const FInstancedPropertyBag& InValueA, const FInstancedPropertyBag& InValueB, FInstancedPropertyBag& OutValue)
+{
+	const FName PropertyHash(InPropertyData.GetLocatorPathHash());
+	const TValueOrError<int32, EPropertyBagResult> ValueAResult = InValueA.GetValueInt32(PropertyHash);
+	const TValueOrError<int32, EPropertyBagResult> ValueBResult = InValueB.GetValueInt32(PropertyHash);
+	if (!ValueAResult.HasValue() || !ValueBResult.HasValue())
+	{
+		return false;
+	}
+
+	const int32 NewValue = ValueAResult.GetValue() - ValueBResult.GetValue();
+	EPropertyBagResult Result = OutValue.SetValueInt32(PropertyHash, NewValue);
+
+	return Result == EPropertyBagResult::Success;
 }
 
 bool UPropertyAnimatorCoreIntHandler::SubtractValue(const FPropertyAnimatorCoreData& InPropertyData, const FInstancedPropertyBag& InValue)
@@ -97,6 +131,24 @@ bool UPropertyAnimatorCoreIntHandler::SubtractValue(const FPropertyAnimatorCoreD
 	return true;
 }
 
+bool UPropertyAnimatorCoreIntHandler::SubtractValue(const FPropertyAnimatorCoreData& InPropertyData, const FInstancedPropertyBag& InValue, FInstancedPropertyBag& OutValue)
+{
+	const FName PropertyHash(InPropertyData.GetLocatorPathHash());
+	const TValueOrError<int32, EPropertyBagResult> ValueResult = InValue.GetValueInt32(PropertyHash);
+	if (!ValueResult.HasValue())
+	{
+		return false;
+	}
+
+	int32 Value;
+	InPropertyData.GetPropertyValuePtr(&Value);
+
+	const int32 NewValue = Value - ValueResult.GetValue();
+	const EPropertyBagResult Result = OutValue.SetValueInt32(PropertyHash, NewValue);
+
+	return Result == EPropertyBagResult::Success;
+}
+
 bool UPropertyAnimatorCoreIntHandler::GetDefaultValue(const FPropertyAnimatorCoreData& InPropertyData, FInstancedPropertyBag& OutValue)
 {
 	const FName PropertyHash(InPropertyData.GetLocatorPathHash());
diff --git a/Engine/Plugins/VirtualProduction/PropertyAnimatorCore/Source/PropertyAnimatorCore/Private/Properties/Handlers/PropertyAnimatorCoreIntHandler.h b/Engine/Plugins/VirtualProduction/PropertyAnimatorCore/Source/PropertyAnimatorCore/Private/Properties/Handlers/PropertyAnimatorCoreIntHandler.h
--- a/Engine/Plugins/VirtualProduction/PropertyAnimatorCore/Source/PropertyAnimatorCore/Private/Properties/Handlers/PropertyAnimatorCoreIntHandler.h
+++ b/Engine/Plugins/VirtualProduction/PropertyAnimatorCore/Source/PropertyAnimatorCore/Private/Properties/Handlers/PropertyAnimatorCoreIntHandler.h
@@ -18,8 +18,11 @@ public:
 	virtual bool SetValue(const FPropertyAnimatorCoreData& InPropertyData, const FInstancedPropertyBag& InValue) override;
 	virtual bool IsAdditiveSupported() const override;
 	virtual bool AddValue(const FPropertyAnimatorCoreData& InPropertyData, const FInstancedPropertyBag& InValue) override;
-	virtual bool AddValue(const FPropertyAnimatorCoreData& InPropertyData, const FInstancedPropertyBag& InOriginalValue, const FInstancedPropertyBag& InDeltaValue) override;
+	virtual bool AddValue(const FPropertyAnimatorCoreData& InPropertyData, const FInstancedPropertyBag& InValueA, const FInstancedPropertyBag&InValueB, FInstancedPropertyBag& OutValue) override;
+	virtual bool AddValue(const FPropertyAnimatorCoreData& InPropertyData, const FInstancedPropertyBag& InValue, FInstancedPropertyBag& OutValue) override;
+	virtual bool SubtractValue(const FPropertyAnimatorCoreData& InPropertyData, const FInstancedPropertyBag& InValueA, const FInstancedPropertyBag& InValueB, FInstancedPropertyBag& OutValue) override;
 	virtual bool SubtractValue(const FPropertyAnimatorCoreData& InPropertyData, const FInstancedPropertyBag& InValue) override;
+	virtual bool SubtractValue(const FPropertyAnimatorCoreData& InPropertyData, const FInstancedPropertyBag& InValue, FInstancedPropertyBag& OutValue) override;
 	virtual bool GetDefaultValue(const FPropertyAnimatorCoreData& InPropertyData, FInstancedPropertyBag& OutValue) override;
 	//~ End UPropertyAnimatorCoreHandlerBase
 };
\ No newline at end of file
diff --git a/Engine/Plugins/VirtualProduction/PropertyAnimatorCore/Source/PropertyAnimatorCore/Private/Properties/Handlers/PropertyAnimatorCoreRotatorHandler.cpp b/Engine/Plugins/VirtualProduction/PropertyAnimatorCore/Source/PropertyAnimatorCore/Private/Properties/Handlers/PropertyAnimatorCoreRotatorHandler.cpp
--- a/Engine/Plugins/VirtualProduction/PropertyAnimatorCore/Source/PropertyAnimatorCore/Private/Properties/Handlers/PropertyAnimatorCoreRotatorHandler.cpp
+++ b/Engine/Plugins/VirtualProduction/PropertyAnimatorCore/Source/PropertyAnimatorCore/Private/Properties/Handlers/PropertyAnimatorCoreRotatorHandler.cpp
@@ -63,20 +63,54 @@ bool UPropertyAnimatorCoreRotatorHandler::AddValue(const FPropertyAnimatorCoreDa
 	return true;
 }
 
-bool UPropertyAnimatorCoreRotatorHandler::AddValue(const FPropertyAnimatorCoreData& InPropertyData, const FInstancedPropertyBag& InOriginalValue, const FInstancedPropertyBag& InDeltaValue)
+bool UPropertyAnimatorCoreRotatorHandler::AddValue(const FPropertyAnimatorCoreData& InPropertyData, const FInstancedPropertyBag& InValueA, const FInstancedPropertyBag&InValueB, FInstancedPropertyBag& OutValue)
 {
 	const FName PropertyHash(InPropertyData.GetLocatorPathHash());
-	const TValueOrError<FRotator*, EPropertyBagResult> OriginalValueResult = InOriginalValue.GetValueStruct<FRotator>(PropertyHash);
-	const TValueOrError<FRotator*, EPropertyBagResult> DeltaValueResult = InDeltaValue.GetValueStruct<FRotator>(PropertyHash);
-	if (!OriginalValueResult.HasValue() || !DeltaValueResult.HasValue())
+	const TValueOrError<FRotator*, EPropertyBagResult> ValueAResult = InValueA.GetValueStruct<FRotator>(PropertyHash);
+	const TValueOrError<FRotator*, EPropertyBagResult> ValueBResult = InValueB.GetValueStruct<FRotator>(PropertyHash);
+	if (!ValueAResult.HasValue() || !ValueBResult.HasValue())
 	{
 		return false;
 	}
 
-	FRotator NewValue = FRotator(OriginalValueResult.GetValue()->Quaternion() * DeltaValueResult.GetValue()->Quaternion());
-	InPropertyData.SetPropertyValuePtr(&NewValue);
+	const FRotator NewValue = FRotator(ValueAResult.GetValue()->Quaternion() * ValueBResult.GetValue()->Quaternion());
+	EPropertyBagResult Result = OutValue.SetValueStruct(PropertyHash, NewValue);
 
-	return true;
+	return Result == EPropertyBagResult::Success;
+}
+
+bool UPropertyAnimatorCoreRotatorHandler::AddValue(const FPropertyAnimatorCoreData& InPropertyData, const FInstancedPropertyBag& InValue, FInstancedPropertyBag& OutValue)
+{
+	const FName PropertyHash(InPropertyData.GetLocatorPathHash());
+	const TValueOrError<FRotator*, EPropertyBagResult> ValueResult = InValue.GetValueStruct<FRotator>(PropertyHash);
+	if (!ValueResult.HasValue())
+	{
+		return false;
+	}
+
+	FRotator Value;
+	InPropertyData.GetPropertyValuePtr(&Value);
+
+	const FRotator NewValue = FRotator(Value.Quaternion() * ValueResult.GetValue()->Quaternion());
+	EPropertyBagResult Result = OutValue.SetValueStruct(PropertyHash, NewValue);
+
+	return Result == EPropertyBagResult::Success;
+}
+
+bool UPropertyAnimatorCoreRotatorHandler::SubtractValue(const FPropertyAnimatorCoreData& InPropertyData, const FInstancedPropertyBag& InValueA, const FInstancedPropertyBag& InValueB, FInstancedPropertyBag& OutValue)
+{
+	const FName PropertyHash(InPropertyData.GetLocatorPathHash());
+	const TValueOrError<FRotator*, EPropertyBagResult> ValueAResult = InValueA.GetValueStruct<FRotator>(PropertyHash);
+	const TValueOrError<FRotator*, EPropertyBagResult> ValueBResult = InValueB.GetValueStruct<FRotator>(PropertyHash);
+	if (!ValueAResult.HasValue() || !ValueBResult.HasValue())
+	{
+		return false;
+	}
+
+	const FRotator NewValue = FRotator(ValueAResult.GetValue()->Quaternion() * ValueBResult.GetValue()->Quaternion().Inverse());
+	EPropertyBagResult Result = OutValue.SetValueStruct(PropertyHash, NewValue);
+
+	return Result == EPropertyBagResult::Success;
 }
 
 bool UPropertyAnimatorCoreRotatorHandler::SubtractValue(const FPropertyAnimatorCoreData& InPropertyData, const FInstancedPropertyBag& InValue)
@@ -97,6 +131,24 @@ bool UPropertyAnimatorCoreRotatorHandler::SubtractValue(const FPropertyAnimatorC
 	return true;
 }
 
+bool UPropertyAnimatorCoreRotatorHandler::SubtractValue(const FPropertyAnimatorCoreData& InPropertyData, const FInstancedPropertyBag& InValue, FInstancedPropertyBag& OutValue)
+{
+	const FName PropertyHash(InPropertyData.GetLocatorPathHash());
+	const TValueOrError<FRotator*, EPropertyBagResult> ValueResult = InValue.GetValueStruct<FRotator>(PropertyHash);
+	if (!ValueResult.HasValue())
+	{
+		return false;
+	}
+
+	FRotator Value;
+	InPropertyData.GetPropertyValuePtr(&Value);
+
+	const FRotator NewValue = FRotator(Value.Quaternion() * ValueResult.GetValue()->Quaternion().Inverse());
+	const EPropertyBagResult Result = OutValue.SetValueStruct(PropertyHash, NewValue);
+
+	return Result == EPropertyBagResult::Success;
+}
+
 bool UPropertyAnimatorCoreRotatorHandler::GetDefaultValue(const FPropertyAnimatorCoreData& InPropertyData, FInstancedPropertyBag& OutValue)
 {
 	const FName PropertyHash(InPropertyData.GetLocatorPathHash());
diff --git a/Engine/Plugins/VirtualProduction/PropertyAnimatorCore/Source/PropertyAnimatorCore/Private/Properties/Handlers/PropertyAnimatorCoreRotatorHandler.h b/Engine/Plugins/VirtualProduction/PropertyAnimatorCore/Source/PropertyAnimatorCore/Private/Properties/Handlers/PropertyAnimatorCoreRotatorHandler.h
--- a/Engine/Plugins/VirtualProduction/PropertyAnimatorCore/Source/PropertyAnimatorCore/Private/Properties/Handlers/PropertyAnimatorCoreRotatorHandler.h
+++ b/Engine/Plugins/VirtualProduction/PropertyAnimatorCore/Source/PropertyAnimatorCore/Private/Properties/Handlers/PropertyAnimatorCoreRotatorHandler.h
@@ -18,8 +18,11 @@ public:
 	virtual bool SetValue(const FPropertyAnimatorCoreData& InPropertyData, const FInstancedPropertyBag& InValue) override;
 	virtual bool IsAdditiveSupported() const override;
 	virtual bool AddValue(const FPropertyAnimatorCoreData& InPropertyData, const FInstancedPropertyBag& InValue) override;
-	virtual bool AddValue(const FPropertyAnimatorCoreData& InPropertyData, const FInstancedPropertyBag& InOriginalValue, const FInstancedPropertyBag& InDeltaValue) override;
+	virtual bool AddValue(const FPropertyAnimatorCoreData& InPropertyData, const FInstancedPropertyBag& InValueA, const FInstancedPropertyBag&InValueB, FInstancedPropertyBag& OutValue) override;
+	virtual bool AddValue(const FPropertyAnimatorCoreData& InPropertyData, const FInstancedPropertyBag& InValue, FInstancedPropertyBag& OutValue) override;
+	virtual bool SubtractValue(const FPropertyAnimatorCoreData& InPropertyData, const FInstancedPropertyBag& InValueA, const FInstancedPropertyBag& InValueB, FInstancedPropertyBag& OutValue) override;
 	virtual bool SubtractValue(const FPropertyAnimatorCoreData& InPropertyData, const FInstancedPropertyBag& InValue) override;
+	virtual bool SubtractValue(const FPropertyAnimatorCoreData& InPropertyData, const FInstancedPropertyBag& InValue, FInstancedPropertyBag& OutValue) override;
 	virtual bool GetDefaultValue(const FPropertyAnimatorCoreData& InPropertyData, FInstancedPropertyBag& OutValue) override;
 	//~ End UPropertyAnimatorCoreHandlerBase
 };
\ No newline at end of file
diff --git a/Engine/Plugins/VirtualProduction/PropertyAnimatorCore/Source/PropertyAnimatorCore/Private/Properties/Handlers/PropertyAnimatorCoreStringHandler.cpp b/Engine/Plugins/VirtualProduction/PropertyAnimatorCore/Source/PropertyAnimatorCore/Private/Properties/Handlers/PropertyAnimatorCoreStringHandler.cpp
--- a/Engine/Plugins/VirtualProduction/PropertyAnimatorCore/Source/PropertyAnimatorCore/Private/Properties/Handlers/PropertyAnimatorCoreStringHandler.cpp
+++ b/Engine/Plugins/VirtualProduction/PropertyAnimatorCore/Source/PropertyAnimatorCore/Private/Properties/Handlers/PropertyAnimatorCoreStringHandler.cpp
@@ -40,72 +40,6 @@ bool UPropertyAnimatorCoreStringHandler::SetValue(const FPropertyAnimatorCoreDat
 	return true;
 }
 
-bool UPropertyAnimatorCoreStringHandler::IsAdditiveSupported() const
-{
-	return true;
-}
-
-bool UPropertyAnimatorCoreStringHandler::AddValue(const FPropertyAnimatorCoreData& InPropertyData, const FInstancedPropertyBag& InValue)
-{
-	const FName PropertyHash(InPropertyData.GetLocatorPathHash());
-	const TValueOrError<FString, EPropertyBagResult> ValueResult = InValue.GetValueString(PropertyHash);
-	if (!ValueResult.HasValue())
-	{
-		return false;
-	}
-
-	FString Value;
-	InPropertyData.GetPropertyValuePtr(&Value);
-
-	// Append
-	FString NewValue = Value + ValueResult.GetValue();
-	InPropertyData.SetPropertyValuePtr(&NewValue);
-
-	return true;
-}
-
-bool UPropertyAnimatorCoreStringHandler::AddValue(const FPropertyAnimatorCoreData& InPropertyData, const FInstancedPropertyBag& InOriginalValue, const FInstancedPropertyBag& InDeltaValue)
-{
-	const FName PropertyHash(InPropertyData.GetLocatorPathHash());
-	const TValueOrError<FString, EPropertyBagResult> OriginalValueResult = InOriginalValue.GetValueString(PropertyHash);
-	const TValueOrError<FString, EPropertyBagResult> DeltaValueResult = InDeltaValue.GetValueString(PropertyHash);
-	if (!OriginalValueResult.HasValue() || !DeltaValueResult.HasValue())
-	{
-		return false;
-	}
-
-	// Append
-	FString NewValue = OriginalValueResult.GetValue() + DeltaValueResult.GetValue();
-	InPropertyData.SetPropertyValuePtr(&NewValue);
-
-	return true;
-}
-
-bool UPropertyAnimatorCoreStringHandler::SubtractValue(const FPropertyAnimatorCoreData& InPropertyData, const FInstancedPropertyBag& InValue)
-{
-	const FName PropertyHash(InPropertyData.GetLocatorPathHash());
-	const TValueOrError<FString, EPropertyBagResult> ValueResult = InValue.GetValueString(PropertyHash);
-	if (!ValueResult.HasValue())
-	{
-		return false;
-	}
-
-	FString Value;
-	InPropertyData.GetPropertyValuePtr(&Value);
-
-	// Trim end
-	FString NewValue = Value;
-	const FString TrimString = ValueResult.GetValue();
-	if (NewValue.EndsWith(TrimString))
-	{
-		NewValue.LeftChopInline(TrimString.Len());
-	}
-
-	InPropertyData.SetPropertyValuePtr(&NewValue);
-
-	return true;
-}
-
 bool UPropertyAnimatorCoreStringHandler::GetDefaultValue(const FPropertyAnimatorCoreData& InPropertyData, FInstancedPropertyBag& OutValue)
 {
 	const FName PropertyHash(InPropertyData.GetLocatorPathHash());
diff --git a/Engine/Plugins/VirtualProduction/PropertyAnimatorCore/Source/PropertyAnimatorCore/Private/Properties/Handlers/PropertyAnimatorCoreStringHandler.h b/Engine/Plugins/VirtualProduction/PropertyAnimatorCore/Source/PropertyAnimatorCore/Private/Properties/Handlers/PropertyAnimatorCoreStringHandler.h
--- a/Engine/Plugins/VirtualProduction/PropertyAnimatorCore/Source/PropertyAnimatorCore/Private/Properties/Handlers/PropertyAnimatorCoreStringHandler.h
+++ b/Engine/Plugins/VirtualProduction/PropertyAnimatorCore/Source/PropertyAnimatorCore/Private/Properties/Handlers/PropertyAnimatorCoreStringHandler.h
@@ -16,10 +16,6 @@ public:
 	virtual bool IsPropertySupported(const FPropertyAnimatorCoreData& InPropertyData) const override;
 	virtual bool GetValue(const FPropertyAnimatorCoreData& InPropertyData, FInstancedPropertyBag& OutValue) override;
 	virtual bool SetValue(const FPropertyAnimatorCoreData& InPropertyData, const FInstancedPropertyBag& InValue) override;
-	virtual bool IsAdditiveSupported() const override;
-	virtual bool AddValue(const FPropertyAnimatorCoreData& InPropertyData, const FInstancedPropertyBag& InValue) override;
-	virtual bool AddValue(const FPropertyAnimatorCoreData& InPropertyData, const FInstancedPropertyBag& InOriginalValue, const FInstancedPropertyBag& InDeltaValue) override;
-	virtual bool SubtractValue(const FPropertyAnimatorCoreData& InPropertyData, const FInstancedPropertyBag& InValue) override;
 	virtual bool GetDefaultValue(const FPropertyAnimatorCoreData& InPropertyData, FInstancedPropertyBag& OutValue) override;
 	//~ End UPropertyAnimatorCoreHandlerBase
 };
\ No newline at end of file
diff --git a/Engine/Plugins/VirtualProduction/PropertyAnimatorCore/Source/PropertyAnimatorCore/Private/Properties/Handlers/PropertyAnimatorCoreTextHandler.cpp b/Engine/Plugins/VirtualProduction/PropertyAnimatorCore/Source/PropertyAnimatorCore/Private/Properties/Handlers/PropertyAnimatorCoreTextHandler.cpp
--- a/Engine/Plugins/VirtualProduction/PropertyAnimatorCore/Source/PropertyAnimatorCore/Private/Properties/Handlers/PropertyAnimatorCoreTextHandler.cpp
+++ b/Engine/Plugins/VirtualProduction/PropertyAnimatorCore/Source/PropertyAnimatorCore/Private/Properties/Handlers/PropertyAnimatorCoreTextHandler.cpp
@@ -42,73 +42,6 @@ bool UPropertyAnimatorCoreTextHandler::SetValue(const FPropertyAnimatorCoreData&
 	return true;
 }
 
-bool UPropertyAnimatorCoreTextHandler::IsAdditiveSupported() const
-{
-	return true;
-}
-
-bool UPropertyAnimatorCoreTextHandler::AddValue(const FPropertyAnimatorCoreData& InPropertyData, const FInstancedPropertyBag& InValue)
-{
-	const FName PropertyHash(InPropertyData.GetLocatorPathHash());
-	const TValueOrError<FText, EPropertyBagResult> ValueResult = InValue.GetValueText(PropertyHash);
-	if (!ValueResult.HasValue())
-	{
-		return false;
-	}
-
-	FText Value;
-	InPropertyData.GetPropertyValuePtr(&Value);
-
-	// Append
-	FText NewValue = FText::Join(FText::GetEmpty(), Value, ValueResult.GetValue());
-	InPropertyData.SetPropertyValuePtr(&NewValue);
-
-	return true;
-}
-
-bool UPropertyAnimatorCoreTextHandler::AddValue(const FPropertyAnimatorCoreData& InPropertyData, const FInstancedPropertyBag& InOriginalValue, const FInstancedPropertyBag& InDeltaValue)
-{
-	const FName PropertyHash(InPropertyData.GetLocatorPathHash());
-	const TValueOrError<FText, EPropertyBagResult> OriginalValueResult = InOriginalValue.GetValueText(PropertyHash);
-	const TValueOrError<FText, EPropertyBagResult> DeltaValueResult = InDeltaValue.GetValueText(PropertyHash);
-	if (!OriginalValueResult.HasValue() || !DeltaValueResult.HasValue())
-	{
-		return false;
-	}
-
-	// Append
-	FText NewValue = FText::Join(FText::GetEmpty(), OriginalValueResult.GetValue(), DeltaValueResult.GetValue());
-	InPropertyData.SetPropertyValuePtr(&NewValue);
-
-	return true;
-}
-
-bool UPropertyAnimatorCoreTextHandler::SubtractValue(const FPropertyAnimatorCoreData& InPropertyData, const FInstancedPropertyBag& InValue)
-{
-	const FName PropertyHash(InPropertyData.GetLocatorPathHash());
-	const TValueOrError<FText, EPropertyBagResult> ValueResult = InValue.GetValueText(PropertyHash);
-	if (!ValueResult.HasValue())
-	{
-		return false;
-	}
-
-	FText Value;
-	InPropertyData.GetPropertyValuePtr(&Value);
-
-	// Trim end
-	FString StringValue = Value.ToString();
-	const FString StringValueResult = ValueResult.GetValue().ToString();
-	if (StringValue.EndsWith(StringValueResult))
-	{
-		StringValue.LeftChopInline(StringValueResult.Len());
-	}
-
-	FText NewValue = FText::FromString(StringValue);
-	InPropertyData.SetPropertyValuePtr(&NewValue);
-
-	return true;
-}
-
 bool UPropertyAnimatorCoreTextHandler::GetDefaultValue(const FPropertyAnimatorCoreData& InPropertyData, FInstancedPropertyBag& OutValue)
 {
 	const FName PropertyHash(InPropertyData.GetLocatorPathHash());
diff --git a/Engine/Plugins/VirtualProduction/PropertyAnimatorCore/Source/PropertyAnimatorCore/Private/Properties/Handlers/PropertyAnimatorCoreTextHandler.h b/Engine/Plugins/VirtualProduction/PropertyAnimatorCore/Source/PropertyAnimatorCore/Private/Properties/Handlers/PropertyAnimatorCoreTextHandler.h
--- a/Engine/Plugins/VirtualProduction/PropertyAnimatorCore/Source/PropertyAnimatorCore/Private/Properties/Handlers/PropertyAnimatorCoreTextHandler.h
+++ b/Engine/Plugins/VirtualProduction/PropertyAnimatorCore/Source/PropertyAnimatorCore/Private/Properties/Handlers/PropertyAnimatorCoreTextHandler.h
@@ -16,10 +16,6 @@ public:
 	virtual bool IsPropertySupported(const FPropertyAnimatorCoreData& InPropertyData) const override;
 	virtual bool GetValue(const FPropertyAnimatorCoreData& InPropertyData, FInstancedPropertyBag& OutValue) override;
 	virtual bool SetValue(const FPropertyAnimatorCoreData& InPropertyData, const FInstancedPropertyBag& InValue) override;
-	virtual bool IsAdditiveSupported() const override;
-	virtual bool AddValue(const FPropertyAnimatorCoreData& InPropertyData, const FInstancedPropertyBag& InValue) override;
-	virtual bool AddValue(const FPropertyAnimatorCoreData& InPropertyData, const FInstancedPropertyBag& InOriginalValue, const FInstancedPropertyBag& InDeltaValue) override;
-	virtual bool SubtractValue(const FPropertyAnimatorCoreData& InPropertyData, const FInstancedPropertyBag& InValue) override;
 	virtual bool GetDefaultValue(const FPropertyAnimatorCoreData& InPropertyData, FInstancedPropertyBag& OutValue) override;
 	//~ End UPropertyAnimatorCoreHandlerBase
 };
\ No newline at end of file
diff --git a/Engine/Plugins/VirtualProduction/PropertyAnimatorCore/Source/PropertyAnimatorCore/Private/Properties/Handlers/PropertyAnimatorCoreVectorHandler.cpp b/Engine/Plugins/VirtualProduction/PropertyAnimatorCore/Source/PropertyAnimatorCore/Private/Properties/Handlers/PropertyAnimatorCoreVectorHandler.cpp
--- a/Engine/Plugins/VirtualProduction/PropertyAnimatorCore/Source/PropertyAnimatorCore/Private/Properties/Handlers/PropertyAnimatorCoreVectorHandler.cpp
+++ b/Engine/Plugins/VirtualProduction/PropertyAnimatorCore/Source/PropertyAnimatorCore/Private/Properties/Handlers/PropertyAnimatorCoreVectorHandler.cpp
@@ -63,20 +63,54 @@ bool UPropertyAnimatorCoreVectorHandler::AddValue(const FPropertyAnimatorCoreDat
 	return true;
 }
 
-bool UPropertyAnimatorCoreVectorHandler::AddValue(const FPropertyAnimatorCoreData& InPropertyData, const FInstancedPropertyBag& InOriginalValue, const FInstancedPropertyBag& InDeltaValue)
+bool UPropertyAnimatorCoreVectorHandler::AddValue(const FPropertyAnimatorCoreData& InPropertyData, const FInstancedPropertyBag& InValueA, const FInstancedPropertyBag& InValueB, FInstancedPropertyBag& OutValue)
 {
 	const FName PropertyHash(InPropertyData.GetLocatorPathHash());
-	const TValueOrError<FVector*, EPropertyBagResult> OriginalValueResult = InOriginalValue.GetValueStruct<FVector>(PropertyHash);
-	const TValueOrError<FVector*, EPropertyBagResult> DeltaValueResult = InOriginalValue.GetValueStruct<FVector>(PropertyHash);
-	if (!OriginalValueResult.HasValue() || !DeltaValueResult.HasValue())
+	const TValueOrError<FVector*, EPropertyBagResult> ValueAResult = InValueA.GetValueStruct<FVector>(PropertyHash);
+	const TValueOrError<FVector*, EPropertyBagResult> ValueBResult = InValueB.GetValueStruct<FVector>(PropertyHash);
+	if (!ValueAResult.HasValue() || !ValueBResult.HasValue())
 	{
 		return false;
 	}
 
-	FVector NewValue = *OriginalValueResult.GetValue() + *DeltaValueResult.GetValue();
-	InPropertyData.SetPropertyValuePtr(&NewValue);
+	const FVector NewValue = *ValueAResult.GetValue() + *ValueBResult.GetValue();
+	const EPropertyBagResult Result = OutValue.SetValueStruct(PropertyHash, NewValue);
 
-	return true;
+	return Result == EPropertyBagResult::Success;
+}
+
+bool UPropertyAnimatorCoreVectorHandler::AddValue(const FPropertyAnimatorCoreData& InPropertyData, const FInstancedPropertyBag& InValue, FInstancedPropertyBag& OutValue)
+{
+	const FName PropertyHash(InPropertyData.GetLocatorPathHash());
+	const TValueOrError<FVector*, EPropertyBagResult> ValueResult = InValue.GetValueStruct<FVector>(PropertyHash);
+	if (!ValueResult.HasValue())
+	{
+		return false;
+	}
+
+	FVector Value;
+	InPropertyData.GetPropertyValuePtr(&Value);
+
+	const FVector NewValue = Value + *ValueResult.GetValue();
+	EPropertyBagResult Result = OutValue.SetValueStruct(PropertyHash, NewValue);
+
+	return Result == EPropertyBagResult::Success;
+}
+
+bool UPropertyAnimatorCoreVectorHandler::SubtractValue(const FPropertyAnimatorCoreData& InPropertyData, const FInstancedPropertyBag& InValueA, const FInstancedPropertyBag& InValueB, FInstancedPropertyBag& OutValue)
+{
+	const FName PropertyHash(InPropertyData.GetLocatorPathHash());
+	const TValueOrError<FVector*, EPropertyBagResult> ValueAResult = InValueA.GetValueStruct<FVector>(PropertyHash);
+	const TValueOrError<FVector*, EPropertyBagResult> ValueBResult = InValueB.GetValueStruct<FVector>(PropertyHash);
+	if (!ValueAResult.HasValue() || !ValueBResult.HasValue())
+	{
+		return false;
+	}
+
+	const FVector NewValue = *ValueAResult.GetValue() - *ValueBResult.GetValue();
+	EPropertyBagResult Result = OutValue.SetValueStruct(PropertyHash, NewValue);
+
+	return Result == EPropertyBagResult::Success;
 }
 
 bool UPropertyAnimatorCoreVectorHandler::SubtractValue(const FPropertyAnimatorCoreData& InPropertyData, const FInstancedPropertyBag& InValue)
@@ -97,6 +131,24 @@ bool UPropertyAnimatorCoreVectorHandler::SubtractValue(const FPropertyAnimatorCo
 	return true;
 }
 
+bool UPropertyAnimatorCoreVectorHandler::SubtractValue(const FPropertyAnimatorCoreData& InPropertyData, const FInstancedPropertyBag& InValue, FInstancedPropertyBag& OutValue)
+{
+	const FName PropertyHash(InPropertyData.GetLocatorPathHash());
+	const TValueOrError<FVector*, EPropertyBagResult> ValueResult = InValue.GetValueStruct<FVector>(PropertyHash);
+	if (!ValueResult.HasValue())
+	{
+		return false;
+	}
+
+	FVector Value;
+	InPropertyData.GetPropertyValuePtr(&Value);
+
+	const FVector NewValue = Value - *ValueResult.GetValue();
+	const EPropertyBagResult Result = OutValue.SetValueStruct(PropertyHash, NewValue);
+
+	return Result == EPropertyBagResult::Success;
+}
+
 bool UPropertyAnimatorCoreVectorHandler::GetDefaultValue(const FPropertyAnimatorCoreData& InPropertyData, FInstancedPropertyBag& OutValue)
 {
 	const FName PropertyHash(InPropertyData.GetLocatorPathHash());
diff --git a/Engine/Plugins/VirtualProduction/PropertyAnimatorCore/Source/PropertyAnimatorCore/Private/Properties/Handlers/PropertyAnimatorCoreVectorHandler.h b/Engine/Plugins/VirtualProduction/PropertyAnimatorCore/Source/PropertyAnimatorCore/Private/Properties/Handlers/PropertyAnimatorCoreVectorHandler.h
--- a/Engine/Plugins/VirtualProduction/PropertyAnimatorCore/Source/PropertyAnimatorCore/Private/Properties/Handlers/PropertyAnimatorCoreVectorHandler.h
+++ b/Engine/Plugins/VirtualProduction/PropertyAnimatorCore/Source/PropertyAnimatorCore/Private/Properties/Handlers/PropertyAnimatorCoreVectorHandler.h
@@ -18,8 +18,11 @@ public:
 	virtual bool SetValue(const FPropertyAnimatorCoreData& InPropertyData, const FInstancedPropertyBag& InValue) override;
 	virtual bool IsAdditiveSupported() const override;
 	virtual bool AddValue(const FPropertyAnimatorCoreData& InPropertyData, const FInstancedPropertyBag& InValue) override;
-	virtual bool AddValue(const FPropertyAnimatorCoreData& InPropertyData, const FInstancedPropertyBag& InOriginalValue, const FInstancedPropertyBag& InDeltaValue) override;
+	virtual bool AddValue(const FPropertyAnimatorCoreData& InPropertyData, const FInstancedPropertyBag& InValueA, const FInstancedPropertyBag& InValueB, FInstancedPropertyBag& OutValue) override;
+	virtual bool AddValue(const FPropertyAnimatorCoreData& InPropertyData, const FInstancedPropertyBag& InValue, FInstancedPropertyBag& OutValue) override;
+	virtual bool SubtractValue(const FPropertyAnimatorCoreData& InPropertyData, const FInstancedPropertyBag& InValueA, const FInstancedPropertyBag& InValueB, FInstancedPropertyBag& OutValue) override;
 	virtual bool SubtractValue(const FPropertyAnimatorCoreData& InPropertyData, const FInstancedPropertyBag& InValue) override;
+	virtual bool SubtractValue(const FPropertyAnimatorCoreData& InPropertyData, const FInstancedPropertyBag& InValue, FInstancedPropertyBag& OutValue) override;
 	virtual bool GetDefaultValue(const FPropertyAnimatorCoreData& InPropertyData, FInstancedPropertyBag& OutValue) override;
 	//~ End UPropertyAnimatorCoreHandlerBase
 };
\ No newline at end of file
diff --git a/Engine/Plugins/VirtualProduction/PropertyAnimatorCore/Source/PropertyAnimatorCore/Private/Properties/PropertyAnimatorCoreContext.cpp b/Engine/Plugins/VirtualProduction/PropertyAnimatorCore/Source/PropertyAnimatorCore/Private/Properties/PropertyAnimatorCoreContext.cpp
--- a/Engine/Plugins/VirtualProduction/PropertyAnimatorCore/Source/PropertyAnimatorCore/Private/Properties/PropertyAnimatorCoreContext.cpp
+++ b/Engine/Plugins/VirtualProduction/PropertyAnimatorCore/Source/PropertyAnimatorCore/Private/Properties/PropertyAnimatorCoreContext.cpp
@@ -480,7 +480,7 @@ void UPropertyAnimatorCoreContext::OnAnimatedChanged()
 	}
 }
 
-void UPropertyAnimatorCoreContext::CommitEvaluationResult(const FPropertyAnimatorCoreData& InResolvedProperty, const FInstancedPropertyBag& InEvaluatedValues)
+void UPropertyAnimatorCoreContext::CommitEvaluationResult(const FPropertyAnimatorCoreData& InResolvedProperty, FInstancedPropertyBag& InEvaluatedValues)
 {
 	if (!IsAnimated())
 	{
@@ -503,26 +503,26 @@ void UPropertyAnimatorCoreContext::CommitEvaluationResult(const FPropertyAnimato
 
 	if (UPropertyAnimatorCoreConverterBase* Converter = ConverterClass.GetDefaultObject())
 	{
-		if (!Converter->Convert(*FromProperty, InEvaluatedValues, *ToProperty, DeltaPropertyValues, ConverterRule))
+		if (!Converter->Convert(*FromProperty, InEvaluatedValues, *ToProperty, InEvaluatedValues, ConverterRule))
 		{
 			return;
 		}
 	}
-	else
+
+	if (Mode == EPropertyAnimatorCoreMode::Additive)
 	{
-		// Ids need to match for copy to be successful
-		const_cast<FPropertyBagPropertyDesc*>(FromProperty)->ID = ToProperty->ID;
+		// Original = Current - Previous
+		Handler->SubtractValue(InResolvedProperty, DeltaPropertyValues, DeltaPropertyValues);
 
-		DeltaPropertyValues.CopyMatchingValuesByID(InEvaluatedValues);
-	}
+		// NewValue = Original + Evaluated
+		Handler->AddValue(InResolvedProperty, DeltaPropertyValues, InEvaluatedValues, InEvaluatedValues);
 
-	if (Mode == EPropertyAnimatorCoreMode::Absolute)
-	{
-		Handler->SetValue(InResolvedProperty, DeltaPropertyValues);
-		DeltaPropertyValues.RemovePropertyByName(PropertyName);
-	}
-	else
-	{
-		Handler->AddValue(InResolvedProperty, OriginalPropertyValues, DeltaPropertyValues);
+		// Delta = NewValue - Original
+		Handler->SubtractValue(InResolvedProperty, InEvaluatedValues, DeltaPropertyValues, DeltaPropertyValues);
 	}
+
+	Handler->SetValue(InResolvedProperty, InEvaluatedValues);
+
+	// Done with the value
+	InEvaluatedValues.RemovePropertyByName(PropertyName);
 }
diff --git a/Engine/Plugins/VirtualProduction/PropertyAnimatorCore/Source/PropertyAnimatorCore/Public/Properties/Handlers/PropertyAnimatorCoreHandlerBase.h b/Engine/Plugins/VirtualProduction/PropertyAnimatorCore/Source/PropertyAnimatorCore/Public/Properties/Handlers/PropertyAnimatorCoreHandlerBase.h
--- a/Engine/Plugins/VirtualProduction/PropertyAnimatorCore/Source/PropertyAnimatorCore/Public/Properties/Handlers/PropertyAnimatorCoreHandlerBase.h
+++ b/Engine/Plugins/VirtualProduction/PropertyAnimatorCore/Source/PropertyAnimatorCore/Public/Properties/Handlers/PropertyAnimatorCoreHandlerBase.h
@@ -54,8 +54,24 @@ public:
 		return false;
 	}
 
-	/** Adds the value to this property within the owner, delta value and original value are in property bags */
-	virtual bool AddValue(const FPropertyAnimatorCoreData& InPropertyData, const FInstancedPropertyBag& InOriginalValue, const FInstancedPropertyBag& InDeltaValue)
+	/** Adds ValueA to ValueB and stores the result in OutValue */
+	virtual bool AddValue(const FPropertyAnimatorCoreData& InPropertyData, const FInstancedPropertyBag& InValueA, const FInstancedPropertyBag& InValueB, FInstancedPropertyBag& OutValue)
+	{
+		return false;
+	}
+
+	/** Subtracts ValueB from ValueA and stores the result in OutValue */
+	virtual bool SubtractValue(const FPropertyAnimatorCoreData& InPropertyData, const FInstancedPropertyBag& InValueA, const FInstancedPropertyBag& InValueB, FInstancedPropertyBag& OutValue)
+	{
+		return false;
+	}
+
+	virtual bool AddValue(const FPropertyAnimatorCoreData& InPropertyData, const FInstancedPropertyBag& InValue, FInstancedPropertyBag& OutValue)
+	{
+		return false;
+	}
+
+	virtual bool SubtractValue(const FPropertyAnimatorCoreData& InPropertyData, const FInstancedPropertyBag& InValue, FInstancedPropertyBag& OutValue)
 	{
 		return false;
 	}
diff --git a/Engine/Plugins/VirtualProduction/PropertyAnimatorCore/Source/PropertyAnimatorCore/Public/Properties/PropertyAnimatorCoreContext.h b/Engine/Plugins/VirtualProduction/PropertyAnimatorCore/Source/PropertyAnimatorCore/Public/Properties/PropertyAnimatorCoreContext.h
--- a/Engine/Plugins/VirtualProduction/PropertyAnimatorCore/Source/PropertyAnimatorCore/Public/Properties/PropertyAnimatorCoreContext.h
+++ b/Engine/Plugins/VirtualProduction/PropertyAnimatorCore/Source/PropertyAnimatorCore/Public/Properties/PropertyAnimatorCoreContext.h
@@ -135,7 +135,7 @@ private:
 	void OnModeChanged();
 
 	/** Sets the evaluation result for the resolved property */
-	PROPERTYANIMATORCORE_API void CommitEvaluationResult(const FPropertyAnimatorCoreData& InResolvedProperty, const FInstancedPropertyBag& InEvaluatedValues);
+	PROPERTYANIMATORCORE_API void CommitEvaluationResult(const FPropertyAnimatorCoreData& InResolvedProperty, FInstancedPropertyBag& InEvaluatedValues);
 
 	/** Use this to resolve virtual linked property */
 	PROPERTYANIMATORCORE_API TArray<FPropertyAnimatorCoreData> ResolveProperty(bool bInForEvaluation) const;
