MotionDesign : PropertyAnimator
- Deprecating legacy bounce, oscillate, pulse animator (can't be added from menu)
- Fixed allowing preset animator to be added from the advanced context menu of details panel

diff --git a/Engine/Plugins/VirtualProduction/PropertyAnimator/Source/PropertyAnimator/Public/Animators/PropertyAnimatorBounce.h b/Engine/Plugins/VirtualProduction/PropertyAnimator/Source/PropertyAnimator/Public/Animators/PropertyAnimatorBounce.h
--- a/Engine/Plugins/VirtualProduction/PropertyAnimator/Source/PropertyAnimator/Public/Animators/PropertyAnimatorBounce.h
+++ b/Engine/Plugins/VirtualProduction/PropertyAnimator/Source/PropertyAnimator/Public/Animators/PropertyAnimatorBounce.h
@@ -8,7 +8,7 @@
 /**
  * Applies an additive bounce movement with various options on supported float properties
  */
-UCLASS(MinimalAPI, AutoExpandCategories=("Animator"))
+UCLASS(MinimalAPI, AutoExpandCategories=("Animator"), Hidden, meta=(Deprecated, DeprecationMessage="Use preset bounce animator based on curve instead"))
 class UPropertyAnimatorBounce : public UPropertyAnimatorNumericBase
 {
 	GENERATED_BODY()
diff --git a/Engine/Plugins/VirtualProduction/PropertyAnimator/Source/PropertyAnimator/Public/Animators/PropertyAnimatorOscillate.h b/Engine/Plugins/VirtualProduction/PropertyAnimator/Source/PropertyAnimator/Public/Animators/PropertyAnimatorOscillate.h
--- a/Engine/Plugins/VirtualProduction/PropertyAnimator/Source/PropertyAnimator/Public/Animators/PropertyAnimatorOscillate.h
+++ b/Engine/Plugins/VirtualProduction/PropertyAnimator/Source/PropertyAnimator/Public/Animators/PropertyAnimatorOscillate.h
@@ -19,7 +19,7 @@ enum class EPropertyAnimatorOscillateFunction : uint8
 /**
  * Applies an additive regular oscillate movement with various options on supported float properties
  */
-UCLASS(MinimalAPI, AutoExpandCategories=("Animator"))
+UCLASS(MinimalAPI, AutoExpandCategories=("Animator"), Hidden, meta=(Deprecated, DeprecationMessage="Use preset oscillate animator based on curve instead"))
 class UPropertyAnimatorOscillate : public UPropertyAnimatorNumericBase
 {
 	GENERATED_BODY()
diff --git a/Engine/Plugins/VirtualProduction/PropertyAnimator/Source/PropertyAnimator/Public/Animators/PropertyAnimatorPulse.h b/Engine/Plugins/VirtualProduction/PropertyAnimator/Source/PropertyAnimator/Public/Animators/PropertyAnimatorPulse.h
--- a/Engine/Plugins/VirtualProduction/PropertyAnimator/Source/PropertyAnimator/Public/Animators/PropertyAnimatorPulse.h
+++ b/Engine/Plugins/VirtualProduction/PropertyAnimator/Source/PropertyAnimator/Public/Animators/PropertyAnimatorPulse.h
@@ -11,7 +11,7 @@
 /**
  * Applies an additive pulse movement with various options on supported float properties
  */
-UCLASS(MinimalAPI, AutoExpandCategories=("Animator"))
+UCLASS(MinimalAPI, AutoExpandCategories=("Animator"), Hidden, meta=(Deprecated, DeprecationMessage="Use preset pulse animator based on curve instead"))
 class UPropertyAnimatorPulse : public UPropertyAnimatorNumericBase
 {
 	GENERATED_BODY()
diff --git a/Engine/Plugins/VirtualProduction/PropertyAnimatorCore/Source/PropertyAnimatorCore/Private/Subsystems/PropertyAnimatorCoreSubsystem.cpp b/Engine/Plugins/VirtualProduction/PropertyAnimatorCore/Source/PropertyAnimatorCore/Private/Subsystems/PropertyAnimatorCoreSubsystem.cpp
--- a/Engine/Plugins/VirtualProduction/PropertyAnimatorCore/Source/PropertyAnimatorCore/Private/Subsystems/PropertyAnimatorCoreSubsystem.cpp
+++ b/Engine/Plugins/VirtualProduction/PropertyAnimatorCore/Source/PropertyAnimatorCore/Private/Subsystems/PropertyAnimatorCoreSubsystem.cpp
@@ -101,7 +101,12 @@ bool UPropertyAnimatorCoreSubsystem::RegisterAnimatorClass(const UClass* InAnima
 	}
 
 	if (!InAnimatorClass->IsChildOf(UPropertyAnimatorCoreBase::StaticClass())
-		|| InAnimatorClass->HasAnyClassFlags(CLASS_Abstract | CLASS_Deprecated | CLASS_NewerVersionExists))
+		|| InAnimatorClass->HasAnyClassFlags(CLASS_Abstract | CLASS_Deprecated | CLASS_NewerVersionExists)
+#if WITH_EDITOR
+		|| InAnimatorClass->HasMetaData(TEXT("Hidden"))
+		|| InAnimatorClass->HasMetaData(TEXT("Deprecated"))
+#endif
+		)
 	{
 		return false;
 	}
@@ -336,6 +341,11 @@ UPropertyAnimatorCoreBase* UPropertyAnimatorCoreSubsystem::CreateAnimator(AActor
 }
 
 TSet<UPropertyAnimatorCoreBase*> UPropertyAnimatorCoreSubsystem::CreateAnimators(const TSet<AActor*>& InActors, const UClass* InAnimatorClass, UPropertyAnimatorCorePresetBase* InPreset, bool bInTransact) const
+{
+	return CreateAnimators(InActors, InAnimatorClass, TSet<UPropertyAnimatorCorePresetBase*>{InPreset}, bInTransact);
+}
+
+TSet<UPropertyAnimatorCoreBase*> UPropertyAnimatorCoreSubsystem::CreateAnimators(const TSet<AActor*>& InActors, const UClass* InAnimatorClass, const TSet<UPropertyAnimatorCorePresetBase*>& InPresets, bool bInTransact) const
 {
 	TSet<UPropertyAnimatorCoreBase*> NewAnimators;
 
@@ -385,7 +395,7 @@ TSet<UPropertyAnimatorCoreBase*> UPropertyAnimatorCoreSubsystem::CreateAnimators
 		}
 
 #if WITH_EDITOR
-      	Component->Modify();
+		Component->Modify();
 #endif
 
 		UPropertyAnimatorCoreBase* NewActorAnimator = Component->AddAnimator(InAnimatorClass);
@@ -399,10 +409,13 @@ TSet<UPropertyAnimatorCoreBase*> UPropertyAnimatorCoreSubsystem::CreateAnimators
 		NewActorAnimator->Modify();
 #endif
 
-		// Optionally apply preset if any
-		if (InPreset && InPreset->IsPresetSupported(Actor, NewActorAnimator))
+		// Optionally apply presets if any
+		for (UPropertyAnimatorCorePresetBase* Preset : InPresets)
 		{
-			InPreset->ApplyPreset(NewActorAnimator);
+			if (Preset && Preset->IsPresetSupported(Actor, NewActorAnimator))
+			{
+				Preset->ApplyPreset(NewActorAnimator);
+			}
 		}
 
 		NewAnimators.Add(NewActorAnimator);
diff --git a/Engine/Plugins/VirtualProduction/PropertyAnimatorCore/Source/PropertyAnimatorCore/Public/Subsystems/PropertyAnimatorCoreSubsystem.h b/Engine/Plugins/VirtualProduction/PropertyAnimatorCore/Source/PropertyAnimatorCore/Public/Subsystems/PropertyAnimatorCoreSubsystem.h
--- a/Engine/Plugins/VirtualProduction/PropertyAnimatorCore/Source/PropertyAnimatorCore/Public/Subsystems/PropertyAnimatorCoreSubsystem.h
+++ b/Engine/Plugins/VirtualProduction/PropertyAnimatorCore/Source/PropertyAnimatorCore/Public/Subsystems/PropertyAnimatorCoreSubsystem.h
@@ -159,6 +159,9 @@ public:
 	/** Create animators of specific class for actors */
 	PROPERTYANIMATORCORE_API TSet<UPropertyAnimatorCoreBase*> CreateAnimators(const TSet<AActor*>& InActors, const UClass* InAnimatorClass, UPropertyAnimatorCorePresetBase* InPreset = nullptr, bool bInTransact = false) const;
 
+	/** Create animators of specific class for actors with presets */
+	PROPERTYANIMATORCORE_API TSet<UPropertyAnimatorCoreBase*> CreateAnimators(const TSet<AActor*>& InActors, const UClass* InAnimatorClass, const TSet<UPropertyAnimatorCorePresetBase*>& InPresets, bool bInTransact = false) const;
+
 	/** Clone animators onto an actor */
 	PROPERTYANIMATORCORE_API TSet<UPropertyAnimatorCoreBase*> CloneAnimators(const TSet<UPropertyAnimatorCoreBase*>& InAnimators, AActor* InTargetActor, bool bInTransact = false) const;
 
diff --git a/Engine/Plugins/VirtualProduction/PropertyAnimatorCore/Source/PropertyAnimatorCoreEditor/Private/Menus/PropertyAnimatorCoreEditorMenu.cpp b/Engine/Plugins/VirtualProduction/PropertyAnimatorCore/Source/PropertyAnimatorCoreEditor/Private/Menus/PropertyAnimatorCoreEditorMenu.cpp
--- a/Engine/Plugins/VirtualProduction/PropertyAnimatorCore/Source/PropertyAnimatorCoreEditor/Private/Menus/PropertyAnimatorCoreEditorMenu.cpp
+++ b/Engine/Plugins/VirtualProduction/PropertyAnimatorCore/Source/PropertyAnimatorCoreEditor/Private/Menus/PropertyAnimatorCoreEditorMenu.cpp
@@ -44,9 +44,9 @@ void UE::PropertyAnimatorCoreEditor::Menu::FillNewAnimatorSection(UToolMenu* InM
 
 	FToolMenuSection& NewAnimatorsSection = InMenu->FindOrAddSection(TEXT("NewAnimators"), LOCTEXT("NewAnimators.Label", "New Animators"));
 
-	TMap<UPropertyAnimatorCoreBase*, TArray<UPropertyAnimatorCorePresetBase*>> AvailablePresetAnimators;
+	TMap<UPropertyAnimatorCoreBase*, TArray<UPropertyAnimatorCoreAnimatorPreset*>> AvailablePresetAnimators;
 
-	UPropertyAnimatorCorePresetBase* EmptyPreset = nullptr;
+	UPropertyAnimatorCoreAnimatorPreset* EmptyPreset = nullptr;
 	for (UPropertyAnimatorCoreBase* NewAnimator : NewAvailableAnimators)
 	{
 		AvailablePresetAnimators.FindOrAdd(NewAnimator).Add(EmptyPreset);
@@ -63,7 +63,7 @@ void UE::PropertyAnimatorCoreEditor::Menu::FillNewAnimatorSection(UToolMenu* InM
 		}
 	}
 
-	for (const TPair<UPropertyAnimatorCoreBase*, TArray<UPropertyAnimatorCorePresetBase*>>& NewPresetAnimator : AvailablePresetAnimators)
+	for (const TPair<UPropertyAnimatorCoreBase*, TArray<UPropertyAnimatorCoreAnimatorPreset*>>& NewPresetAnimator : AvailablePresetAnimators)
 	{
 		UPropertyAnimatorCoreBase* NewAnimator = NewPresetAnimator.Key;
 
@@ -74,28 +74,29 @@ void UE::PropertyAnimatorCoreEditor::Menu::FillNewAnimatorSection(UToolMenu* InM
 		const FText MenuTooltip = LOCTEXT("NewAnimator.Tooltip", "Create a new animator");
 		const FSlateIcon MenuIcon = FSlateIconFinder::FindIconForClass(NewAnimator->GetClass());
 
-		if (bAdvancedMenu)
+		UPropertyAnimatorCorePropertyPreset* EmptyPropertyPreset = nullptr;
+		for (UPropertyAnimatorCoreAnimatorPreset* AnimatorPreset : NewPresetAnimator.Value)
 		{
-			NewAnimatorsSection.AddSubMenu(
-				MenuName
-				, MenuLabel
-				, MenuTooltip
-				, FNewToolMenuDelegate::CreateLambda(&FillNewAnimatorSubmenu, NewAnimator, InMenuData)
-				, bOpenOnClick
-				, MenuIcon
-				, bCloseMenuAfterSelection
-			);
-		}
-		else
-		{
-			for (UPropertyAnimatorCorePresetBase* Preset : NewPresetAnimator.Value)
+			if (AnimatorPreset)
 			{
-				if (Preset)
-				{
-					MenuLabel = FText::Format(LOCTEXT("PresetMenuLabel", "{0} ({1})"), Preset->GetPresetDisplayName(), Metadata->DisplayName);
-					MenuCategory = TEXT("Presets");
-				}
+				MenuLabel = FText::Format(LOCTEXT("PresetMenuLabel", "{0} ({1})"), AnimatorPreset->GetPresetDisplayName(), Metadata->DisplayName);
+				MenuCategory = TEXT("Presets");
+			}
 
+			if (bAdvancedMenu)
+			{
+				NewAnimatorsSection.AddSubMenu(
+					MenuName
+					, MenuLabel
+					, MenuTooltip
+					, FNewToolMenuDelegate::CreateLambda(&FillNewAnimatorSubmenu, NewAnimator, AnimatorPreset, InMenuData)
+					, bOpenOnClick
+					, MenuIcon
+					, bCloseMenuAfterSelection
+				);
+			}
+			else
+			{
 				TSharedPtr<const FPropertyAnimatorCoreEditorCategoryMetadata> AnimatorCategory = EditorSubsystem->FindAnimatorCategory(MenuCategory);
 				FToolMenuSection& AnimatorCategorySection = InMenu->FindOrAddSection(AnimatorCategory ? AnimatorCategory->Name : MenuCategory, AnimatorCategory ? AnimatorCategory->DisplayName : FText::FromName(MenuCategory));
 
@@ -104,7 +105,7 @@ void UE::PropertyAnimatorCoreEditor::Menu::FillNewAnimatorSection(UToolMenu* InM
 					, MenuLabel
 					, MenuTooltip
 					, MenuIcon
-					, FExecuteAction::CreateLambda(&ExecuteNewAnimatorPresetAction, NewAnimator, ContextActors, Preset, InMenuData)
+					, FExecuteAction::CreateLambda(&ExecuteNewAnimatorPresetAction, NewAnimator, ContextActors, AnimatorPreset, EmptyPropertyPreset, InMenuData)
 				);
 			}
 		}
@@ -372,7 +373,7 @@ void UE::PropertyAnimatorCoreEditor::Menu::FillDisableAnimatorSection(UToolMenu*
 	}
 }
 
-void UE::PropertyAnimatorCoreEditor::Menu::ExecuteNewAnimatorPresetAction(const UPropertyAnimatorCoreBase* InAnimator, const TSet<AActor*>& InActors, UPropertyAnimatorCorePresetBase* InPreset, TSharedRef<FPropertyAnimatorCoreEditorMenuData> InMenuData)
+void UE::PropertyAnimatorCoreEditor::Menu::ExecuteNewAnimatorPresetAction(const UPropertyAnimatorCoreBase* InAnimator, const TSet<AActor*>& InActors, UPropertyAnimatorCoreAnimatorPreset* InAnimatorPreset, UPropertyAnimatorCorePropertyPreset* InPropertyPreset, TSharedRef<FPropertyAnimatorCoreEditorMenuData> InMenuData)
 {
 	const UPropertyAnimatorCoreSubsystem* Subsystem = UPropertyAnimatorCoreSubsystem::Get();
 
@@ -384,10 +385,16 @@ void UE::PropertyAnimatorCoreEditor::Menu::ExecuteNewAnimatorPresetAction(const
 		return;
 	}
 
-	InMenuData->SetLastCreatedAnimators(Subsystem->CreateAnimators(InActors, InAnimator->GetClass(), InPreset, InMenuData->GetOptions().ShouldTransact()));
+	const TSet<UPropertyAnimatorCorePresetBase*> Presets
+	{
+		InAnimatorPreset,
+		InPropertyPreset
+	};
+
+	InMenuData->SetLastCreatedAnimators(Subsystem->CreateAnimators(InActors, InAnimator->GetClass(), Presets, InMenuData->GetOptions().ShouldTransact()));
 }
 
-void UE::PropertyAnimatorCoreEditor::Menu::ExecuteNewAnimatorPropertyAction(const UPropertyAnimatorCoreBase* InAnimator, const FString& InPropertyLocatorPath, TSharedRef<FPropertyAnimatorCoreEditorMenuData> InMenuData)
+void UE::PropertyAnimatorCoreEditor::Menu::ExecuteNewAnimatorPropertyAction(const UPropertyAnimatorCoreBase* InAnimator, const FString& InPropertyLocatorPath, UPropertyAnimatorCoreAnimatorPreset* InAnimatorPreset, TSharedRef<FPropertyAnimatorCoreEditorMenuData> InMenuData)
 {
 	UPropertyAnimatorCoreSubsystem* Subsystem = UPropertyAnimatorCoreSubsystem::Get();
 
@@ -410,12 +417,12 @@ void UE::PropertyAnimatorCoreEditor::Menu::ExecuteNewAnimatorPropertyAction(cons
 	TSet<UPropertyAnimatorCoreBase*> CreatedAnimators;
 	for (AActor* Actor : ContextActors)
 	{
-		UPropertyAnimatorCoreBase* NewAnimator = Subsystem->CreateAnimator(Actor, InAnimator->GetClass(), nullptr, bShouldTransact);
+		UPropertyAnimatorCoreBase* NewAnimator = Subsystem->CreateAnimator(Actor, InAnimator->GetClass(), InAnimatorPreset, bShouldTransact);
 		FPropertyAnimatorCoreData Property(Actor, InPropertyLocatorPath);
 		Subsystem->LinkAnimatorProperty(NewAnimator, Property, bShouldTransact);
 		CreatedAnimators.Add(NewAnimator);
 	}
-	
+
 	InMenuData->SetLastCreatedAnimators(CreatedAnimators);
 }
 
@@ -790,7 +797,7 @@ bool UE::PropertyAnimatorCoreEditor::Menu::IsAnimatorPresetLinked(UPropertyAnima
 	return InPreset->IsPresetApplied(InAnimator);
 }
 
-void UE::PropertyAnimatorCoreEditor::Menu::FillNewAnimatorSubmenu(UToolMenu* InMenu, UPropertyAnimatorCoreBase* InAnimator, TSharedRef<FPropertyAnimatorCoreEditorMenuData> InMenuData)
+void UE::PropertyAnimatorCoreEditor::Menu::FillNewAnimatorSubmenu(UToolMenu* InMenu, UPropertyAnimatorCoreBase* InAnimator, UPropertyAnimatorCoreAnimatorPreset* InAnimatorPreset, TSharedRef<FPropertyAnimatorCoreEditorMenuData> InMenuData)
 {
 	const UPropertyAnimatorCoreSubsystem* Subsystem = UPropertyAnimatorCoreSubsystem::Get();
 
@@ -806,9 +813,9 @@ void UE::PropertyAnimatorCoreEditor::Menu::FillNewAnimatorSubmenu(UToolMenu* InM
 	{
 		const TSet<AActor*>& ContextActors = InMenuData->GetContext().GetActors();
 
-		FToolMenuSection& PresetSection = InMenu->FindOrAddSection(TEXT("Presets"), LOCTEXT("NewAnimatorPresetsSection.Label", "Presets"));
+		FToolMenuSection& PresetSection = InMenu->FindOrAddSection(TEXT("PropertyPresets"), LOCTEXT("NewAnimatorPropertyPresetsSection.Label", "Property Presets"));
 
-		UPropertyAnimatorCorePropertyPreset* EmptyPreset = nullptr;
+		UPropertyAnimatorCorePropertyPreset* EmptyPropertyPreset = nullptr;
 
 		PresetSection.AddMenuEntry(
 			TEXT("EmptyPreset")
@@ -816,38 +823,39 @@ void UE::PropertyAnimatorCoreEditor::Menu::FillNewAnimatorSubmenu(UToolMenu* InM
 			, LOCTEXT("NewAnimatorEmptyPresetSection.Tooltip", "Create an empty animator")
 			, FSlateIcon()
 			, FUIAction(
-				FExecuteAction::CreateLambda(&ExecuteNewAnimatorPresetAction, InAnimator, ContextActors, EmptyPreset, InMenuData)
+				FExecuteAction::CreateLambda(&ExecuteNewAnimatorPresetAction, InAnimator, ContextActors, InAnimatorPreset, EmptyPropertyPreset, InMenuData)
 				, FCanExecuteAction()
 				, FIsActionChecked()
 				, FIsActionButtonVisible::CreateLambda(&IsLastAnimatorCreatedActionHidden, InAnimator, InMenuData)
 			)
 		);
 
-		TSet<UPropertyAnimatorCorePresetBase*> SupportedPresets = Subsystem->GetAvailablePresets(UPropertyAnimatorCorePropertyPreset::StaticClass());
+		TSet<UPropertyAnimatorCorePresetBase*> SupportedPropertyPresets = Subsystem->GetAvailablePresets(UPropertyAnimatorCorePropertyPreset::StaticClass());
 
 		for (const AActor* Actor : ContextActors)
 		{
-			SupportedPresets = SupportedPresets.Intersect(Subsystem->GetSupportedPresets(Actor, InAnimator, UPropertyAnimatorCorePropertyPreset::StaticClass()));
+			SupportedPropertyPresets = SupportedPropertyPresets.Intersect(Subsystem->GetSupportedPresets(Actor, InAnimator, UPropertyAnimatorCorePropertyPreset::StaticClass()));
 		}
 
 		constexpr bool bCloseMenuAfterSelection = false;
 		constexpr bool bOpenOnClick = false;
 
-		for (UPropertyAnimatorCorePresetBase* SupportedPreset : SupportedPresets)
+		for (UPropertyAnimatorCorePresetBase* SupportedPropertyPreset : SupportedPropertyPresets)
 		{
-			if (!SupportedPreset)
+			UPropertyAnimatorCorePropertyPreset* PropertyPreset = Cast<UPropertyAnimatorCorePropertyPreset>(SupportedPropertyPreset);
+			if (!PropertyPreset)
 			{
 				continue;
 			}
 
-			const FString MenuName = SupportedPreset->GetPresetName().ToString();
-			const FText MenuLabel = SupportedPreset->GetPresetDisplayName();
+			const FString MenuName = SupportedPropertyPreset->GetPresetName().ToString();
+			const FText MenuLabel = SupportedPropertyPreset->GetPresetDisplayName();
 
 			PresetSection.AddSubMenu(
 				FName(TEXT("Create") + MenuName)
 				, MenuLabel
 				, LOCTEXT("NewAnimatorPresetSection.Tooltip", "Create this animator using this preset")
-				, FNewToolMenuDelegate::CreateLambda(&FillNewPresetAnimatorSubmenu, InAnimator, Cast<UPropertyAnimatorCorePropertyPreset>(SupportedPreset), InMenuData)
+				, FNewToolMenuDelegate::CreateLambda(&FillNewPresetAnimatorSubmenu, InAnimator, InAnimatorPreset, PropertyPreset, InMenuData)
 				, bOpenOnClick
 				, FSlateIcon()
 				, bCloseMenuAfterSelection
@@ -878,7 +886,7 @@ void UE::PropertyAnimatorCoreEditor::Menu::FillNewAnimatorSubmenu(UToolMenu* InM
 				, LOCTEXT("NewAnimatorPropertySection.Tooltip", "Create this animator using this property")
 				, FSlateIcon()
 				, FUIAction(
-					FExecuteAction::CreateLambda(&ExecuteNewAnimatorPropertyAction, InAnimator, PropertyLocatorPath, InMenuData)
+					FExecuteAction::CreateLambda(&ExecuteNewAnimatorPropertyAction, InAnimator, PropertyLocatorPath, InAnimatorPreset, InMenuData)
 					, FCanExecuteAction()
 					, FIsActionChecked()
 					, FIsActionButtonVisible::CreateLambda(&IsLastAnimatorCreatedActionHidden, InAnimator, InMenuData)
@@ -918,7 +926,7 @@ void UE::PropertyAnimatorCoreEditor::Menu::FillLinkAnimatorSubmenu(UToolMenu* In
 		return;
 	}
 
-	FToolMenuSection& PresetSection = InMenu->FindOrAddSection(TEXT("Presets"), LOCTEXT("LinkAnimatorPresetsSection.Label", "Presets"));
+	FToolMenuSection& PresetSection = InMenu->FindOrAddSection(TEXT("PropertyPresets"), LOCTEXT("LinkAnimatorPropertyPresetsSection.Label", "Property Presets"));
 
 	TSet<UPropertyAnimatorCorePresetBase*> SupportedPresets = Subsystem->GetAvailablePresets(UPropertyAnimatorCorePropertyPreset::StaticClass());
 
@@ -1079,7 +1087,7 @@ void UE::PropertyAnimatorCoreEditor::Menu::FillPresetAnimatorSubmenu(UToolMenu*
 	}
 }
 
-void UE::PropertyAnimatorCoreEditor::Menu::FillNewPresetAnimatorSubmenu(UToolMenu* InMenu, UPropertyAnimatorCoreBase* InAnimator, UPropertyAnimatorCorePropertyPreset* InPreset, TSharedRef<FPropertyAnimatorCoreEditorMenuData> InMenuData)
+void UE::PropertyAnimatorCoreEditor::Menu::FillNewPresetAnimatorSubmenu(UToolMenu* InMenu, UPropertyAnimatorCoreBase* InAnimator, UPropertyAnimatorCoreAnimatorPreset* InAnimatorPreset, UPropertyAnimatorCorePropertyPreset* InPropertyPreset, TSharedRef<FPropertyAnimatorCoreEditorMenuData> InMenuData)
 {
 	if (!InMenu)
 	{
@@ -1093,7 +1101,7 @@ void UE::PropertyAnimatorCoreEditor::Menu::FillNewPresetAnimatorSubmenu(UToolMen
 	for (const AActor* ContextActor : ContextActors)
 	{
 		TSet<FPropertyAnimatorCoreData> SupportedProperties;
-		InPreset->GetSupportedPresetProperties(ContextActor, InAnimator, SupportedProperties);
+		InPropertyPreset->GetSupportedPresetProperties(ContextActor, InAnimator, SupportedProperties);
 
 		if (Index++ == 0)
 		{
@@ -1111,7 +1119,7 @@ void UE::PropertyAnimatorCoreEditor::Menu::FillNewPresetAnimatorSubmenu(UToolMen
 		, LOCTEXT("NewAnimatorPresetSection.Tooltip", "Create this animator using this preset")
 		, FSlateIcon()
 		, FUIAction(
-			FExecuteAction::CreateLambda(&ExecuteNewAnimatorPresetAction, InAnimator, ContextActors, InPreset, InMenuData)
+			FExecuteAction::CreateLambda(&ExecuteNewAnimatorPresetAction, InAnimator, ContextActors, InAnimatorPreset, InPropertyPreset, InMenuData)
 			, FCanExecuteAction()
 			, FIsActionChecked()
 			, FIsActionButtonVisible::CreateLambda(&IsLastAnimatorCreatedActionHidden, InAnimator, InMenuData)
@@ -1130,9 +1138,9 @@ void UE::PropertyAnimatorCoreEditor::Menu::FillNewPresetAnimatorSubmenu(UToolMen
 		, LOCTEXT("ApplyLastCreatedAnimatorPresetSection.Tooltip", "Apply this preset to the last created animator")
 		, FSlateIcon()
 		, FUIAction(
-			FExecuteAction::CreateLambda(&ExecuteApplyLastCreatedAnimatorPresetAction, InAnimator, InPreset, InMenuData)
+			FExecuteAction::CreateLambda(&ExecuteApplyLastCreatedAnimatorPresetAction, InAnimator, InPropertyPreset, InMenuData)
 			, FCanExecuteAction()
-			, FGetActionCheckState::CreateLambda(&GetLastAnimatorCreatedPresetState, InAnimator, InPreset, InMenuData)
+			, FGetActionCheckState::CreateLambda(&GetLastAnimatorCreatedPresetState, InAnimator, InPropertyPreset, InMenuData)
 			, FIsActionButtonVisible::CreateLambda(&IsLastAnimatorCreatedActionVisible, InAnimator, InMenuData)
 		)
 		, EUserInterfaceActionType::ToggleButton
@@ -1159,7 +1167,7 @@ void UE::PropertyAnimatorCoreEditor::Menu::FillNewPresetAnimatorSubmenu(UToolMen
 			, LOCTEXT("CreateAnimatorPresetProperty.Tooltip", "Create this animator using this preset property")
 			, FSlateIcon()
 			, FUIAction(
-				FExecuteAction::CreateLambda(&ExecuteNewAnimatorPropertyAction, InAnimator, PropertyLocatorPath, InMenuData)
+				FExecuteAction::CreateLambda(&ExecuteNewAnimatorPropertyAction, InAnimator, PropertyLocatorPath, InAnimatorPreset, InMenuData)
 				, FCanExecuteAction()
 				, FIsActionChecked()
 				, FIsActionButtonVisible::CreateLambda(&IsLastAnimatorCreatedActionHidden, InAnimator, InMenuData)
diff --git a/Engine/Plugins/VirtualProduction/PropertyAnimatorCore/Source/PropertyAnimatorCoreEditor/Private/Menus/PropertyAnimatorCoreEditorMenu.h b/Engine/Plugins/VirtualProduction/PropertyAnimatorCore/Source/PropertyAnimatorCoreEditor/Private/Menus/PropertyAnimatorCoreEditorMenu.h
--- a/Engine/Plugins/VirtualProduction/PropertyAnimatorCore/Source/PropertyAnimatorCoreEditor/Private/Menus/PropertyAnimatorCoreEditorMenu.h
+++ b/Engine/Plugins/VirtualProduction/PropertyAnimatorCore/Source/PropertyAnimatorCoreEditor/Private/Menus/PropertyAnimatorCoreEditorMenu.h
@@ -13,6 +13,7 @@ struct FPropertyAnimatorCoreData;
 class IPropertyHandle;
 class UPropertyAnimatorCoreBase;
 class UPropertyAnimatorCoreEditorMenuContext;
+class UPropertyAnimatorCoreAnimatorPreset;
 class UPropertyAnimatorCorePropertyPreset;
 
 namespace UE::PropertyAnimatorCoreEditor::Menu
@@ -33,19 +34,19 @@ namespace UE::PropertyAnimatorCoreEditor::Menu
 
 	/** Submenus */
 
-	void FillNewAnimatorSubmenu(UToolMenu* InMenu, UPropertyAnimatorCoreBase* InAnimator, TSharedRef<FPropertyAnimatorCoreEditorMenuData> InMenuData);
+	void FillNewAnimatorSubmenu(UToolMenu* InMenu, UPropertyAnimatorCoreBase* InAnimator, UPropertyAnimatorCoreAnimatorPreset* InAnimatorPreset, TSharedRef<FPropertyAnimatorCoreEditorMenuData> InMenuData);
 
 	void FillLinkAnimatorSubmenu(UToolMenu* InMenu, const TSet<UPropertyAnimatorCoreBase*>& InAnimators, TSharedRef<FPropertyAnimatorCoreEditorMenuData> InMenuData);
 
 	void FillPresetAnimatorSubmenu(UToolMenu* InMenu, const TSet<UPropertyAnimatorCoreBase*>& InAnimators, UPropertyAnimatorCorePropertyPreset* InPreset, TSharedRef<FPropertyAnimatorCoreEditorMenuData> InMenuData);
 
-	void FillNewPresetAnimatorSubmenu(UToolMenu* InMenu, UPropertyAnimatorCoreBase* InAnimator, UPropertyAnimatorCorePropertyPreset* InPreset, TSharedRef<FPropertyAnimatorCoreEditorMenuData> InMenuData);
+	void FillNewPresetAnimatorSubmenu(UToolMenu* InMenu, UPropertyAnimatorCoreBase* InAnimator, UPropertyAnimatorCoreAnimatorPreset* InAnimatorPreset, UPropertyAnimatorCorePropertyPreset* InPropertyPreset, TSharedRef<FPropertyAnimatorCoreEditorMenuData> InMenuData);
 
 	/** Execute */
 
-	void ExecuteNewAnimatorPresetAction(const UPropertyAnimatorCoreBase* InAnimator, const TSet<AActor*>& InActors, UPropertyAnimatorCorePresetBase* InPreset, TSharedRef<FPropertyAnimatorCoreEditorMenuData> InMenuData);
+	void ExecuteNewAnimatorPresetAction(const UPropertyAnimatorCoreBase* InAnimator, const TSet<AActor*>& InActors, UPropertyAnimatorCoreAnimatorPreset* InAnimatorPreset, UPropertyAnimatorCorePropertyPreset* InPropertyPreset, TSharedRef<FPropertyAnimatorCoreEditorMenuData> InMenuData);
 
-	void ExecuteNewAnimatorPropertyAction(const UPropertyAnimatorCoreBase* InAnimator, const FString& InPropertyLocatorPath, TSharedRef<FPropertyAnimatorCoreEditorMenuData> InMenuData);
+	void ExecuteNewAnimatorPropertyAction(const UPropertyAnimatorCoreBase* InAnimator, const FString& InPropertyLocatorPath, UPropertyAnimatorCoreAnimatorPreset* InAnimatorPreset, TSharedRef<FPropertyAnimatorCoreEditorMenuData> InMenuData);
 
 	void ExecuteLinkLastCreatedAnimatorPropertyAction(const UPropertyAnimatorCoreBase* InAnimator, const FString& InPropertyLocatorPath, TSharedRef<FPropertyAnimatorCoreEditorMenuData> InMenuData);
 
