MotionDesign : PropertyAnimator
- Fixed SVGImporter actor moving pivot when entering PIE when animator is used

diff --git a/Engine/Plugins/VirtualProduction/PropertyAnimatorCore/Source/PropertyAnimatorCore/Private/Animators/PropertyAnimatorCoreBase.cpp b/Engine/Plugins/VirtualProduction/PropertyAnimatorCore/Source/PropertyAnimatorCore/Private/Animators/PropertyAnimatorCoreBase.cpp
--- a/Engine/Plugins/VirtualProduction/PropertyAnimatorCore/Source/PropertyAnimatorCore/Private/Animators/PropertyAnimatorCoreBase.cpp
+++ b/Engine/Plugins/VirtualProduction/PropertyAnimatorCore/Source/PropertyAnimatorCore/Private/Animators/PropertyAnimatorCoreBase.cpp
@@ -43,6 +43,7 @@ UPropertyAnimatorCoreBase::UPropertyAnimatorCoreBase()
 #if WITH_EDITOR
 		FCoreUObjectDelegates::OnObjectsReplaced.AddUObject(this, &UPropertyAnimatorCoreBase::OnObjectReplaced);
 		FEditorDelegates::PreSaveWorldWithContext.AddUObject(this, &UPropertyAnimatorCoreBase::OnPreSaveWorld);
+		FEditorDelegates::PreBeginPIE.AddUObject(this, &UPropertyAnimatorCoreBase::OnPreBeginPIE);
 #endif
 	}
 }
@@ -142,6 +143,7 @@ void UPropertyAnimatorCoreBase::BeginDestroy()
 #if WITH_EDITOR
 	FCoreUObjectDelegates::OnObjectsReplaced.RemoveAll(this);
 	FEditorDelegates::PreSaveWorldWithContext.RemoveAll(this);
+	FEditorDelegates::PreBeginPIE.RemoveAll(this);
 #endif
 }
 
@@ -631,6 +633,11 @@ void UPropertyAnimatorCoreBase::OnPreSaveWorld(UWorld* InWorld, FObjectPreSaveCo
 
 	RestoreProperties();
 }
+
+void UPropertyAnimatorCoreBase::OnPreBeginPIE(bool bInSimulating)
+{
+	RestoreProperties();
+}
 #endif
 
 void UPropertyAnimatorCoreBase::RestoreProperties()
diff --git a/Engine/Plugins/VirtualProduction/PropertyAnimatorCore/Source/PropertyAnimatorCore/Public/Animators/PropertyAnimatorCoreBase.h b/Engine/Plugins/VirtualProduction/PropertyAnimatorCore/Source/PropertyAnimatorCore/Public/Animators/PropertyAnimatorCoreBase.h
--- a/Engine/Plugins/VirtualProduction/PropertyAnimatorCore/Source/PropertyAnimatorCore/Public/Animators/PropertyAnimatorCoreBase.h
+++ b/Engine/Plugins/VirtualProduction/PropertyAnimatorCore/Source/PropertyAnimatorCore/Public/Animators/PropertyAnimatorCoreBase.h
@@ -85,7 +85,7 @@ public:
 	PROPERTYANIMATORCORE_API static FName GetAnimatorEnabledPropertyName();
 	PROPERTYANIMATORCORE_API static FName GetLinkedPropertiesPropertyName();
 #endif
-
+	
 	PROPERTYANIMATORCORE_API UPropertyAnimatorCoreBase();
 
 	PROPERTYANIMATORCORE_API AActor* GetAnimatorActor() const;
@@ -339,6 +339,9 @@ private:
 #if WITH_EDITOR
 	/** Needed to restore properties and stop animation before world is saved */
 	void OnPreSaveWorld(UWorld* InWorld, FObjectPreSaveContext InContext);
+
+	/** Needed to restore properties and stop animation before PIE is started */
+	void OnPreBeginPIE(bool bInSimulating);
 #endif
 
 	void OnAnimatorEnabledChanged(EPropertyAnimatorCoreUpdateEvent InType);
