Automate 'Animators - Remove all animators'

#rnx

#FailedValidations Preflight

diff --git a/Engine/Plugins/VirtualProduction/PropertyAnimatorCore/Source/PropertyAnimatorCore/Public/Components/PropertyAnimatorCoreComponent.h b/Engine/Plugins/VirtualProduction/PropertyAnimatorCore/Source/PropertyAnimatorCore/Public/Components/PropertyAnimatorCoreComponent.h
--- a/Engine/Plugins/VirtualProduction/PropertyAnimatorCore/Source/PropertyAnimatorCore/Public/Components/PropertyAnimatorCoreComponent.h
+++ b/Engine/Plugins/VirtualProduction/PropertyAnimatorCore/Source/PropertyAnimatorCore/Public/Components/PropertyAnimatorCoreComponent.h
@@ -16,7 +16,7 @@ class UPropertyAnimatorCoreComponent : public UActorComponent
 
 public:
 	/** Create an instance of this component class and adds it to an actor */
-	static UPropertyAnimatorCoreComponent* FindOrAdd(AActor* InActor);
+	PROPERTYANIMATORCORE_API static UPropertyAnimatorCoreComponent* FindOrAdd(AActor* InActor);
 
 #if WITH_EDITOR
 	PROPERTYANIMATORCORE_API static FName GetAnimatorsEnabledPropertyName();
@@ -25,7 +25,7 @@ public:
 
 	UPropertyAnimatorCoreComponent();
 
-	void SetAnimators(const TArray<TObjectPtr<UPropertyAnimatorCoreBase>>& InAnimators);
+	PROPERTYANIMATORCORE_API void SetAnimators(const TArray<TObjectPtr<UPropertyAnimatorCoreBase>>& InAnimators);
 	TConstArrayView<TObjectPtr<UPropertyAnimatorCoreBase>> GetAnimators() const
 	{
 		return PropertyAnimators;
