check for UObjectInitialized  && !IsEngineExitRequested() before running clean-up code that involves StaticClass

#rnx
#Synced-CL 49451892

diff --git a/Engine/Plugins/Experimental/CineCameraRigs/Source/CineCameraRigsEditor/Private/CineCameraRigsEditor.cpp b/Engine/Plugins/Experimental/CineCameraRigs/Source/CineCameraRigsEditor/Private/CineCameraRigsEditor.cpp
--- a/Engine/Plugins/Experimental/CineCameraRigs/Source/CineCameraRigsEditor/Private/CineCameraRigsEditor.cpp
+++ b/Engine/Plugins/Experimental/CineCameraRigs/Source/CineCameraRigsEditor/Private/CineCameraRigsEditor.cpp
@@ -38,7 +38,10 @@ public:
 		if (GUnrealEd)
 		{
 			GUnrealEd->UnregisterComponentVisualizer(CineSplineComponentName);
-			UnregisterCustomizations();
+			if (UObjectInitialized() && !IsEngineExitRequested())
+			{
+				UnregisterCustomizations();
+			}
 		}
 	}
 
diff --git a/Engine/Plugins/VirtualProduction/ActorModifierCore/Source/ActorModifierCoreEditor/Private/ActorModifierCoreEditorModule.cpp b/Engine/Plugins/VirtualProduction/ActorModifierCore/Source/ActorModifierCoreEditor/Private/ActorModifierCoreEditorModule.cpp
--- a/Engine/Plugins/VirtualProduction/ActorModifierCore/Source/ActorModifierCoreEditor/Private/ActorModifierCoreEditorModule.cpp
+++ b/Engine/Plugins/VirtualProduction/ActorModifierCore/Source/ActorModifierCoreEditor/Private/ActorModifierCoreEditorModule.cpp
@@ -18,7 +18,10 @@ void FActorModifierCoreEditorModule::StartupModule()
 
 void FActorModifierCoreEditorModule::ShutdownModule()
 {
-	UnregisterDetailCustomizations();
+	if (UObjectInitialized() && !IsEngineExitRequested())
+	{
+		UnregisterDetailCustomizations();
+	}
 }
 
 void FActorModifierCoreEditorModule::RegisterDetailCustomizations()
diff --git a/Engine/Plugins/VirtualProduction/PropertyAnimatorCore/Source/PropertyAnimatorCoreEditor/Private/PropertyAnimatorCoreEditorModule.cpp b/Engine/Plugins/VirtualProduction/PropertyAnimatorCore/Source/PropertyAnimatorCoreEditor/Private/PropertyAnimatorCoreEditorModule.cpp
--- a/Engine/Plugins/VirtualProduction/PropertyAnimatorCore/Source/PropertyAnimatorCoreEditor/Private/PropertyAnimatorCoreEditorModule.cpp
+++ b/Engine/Plugins/VirtualProduction/PropertyAnimatorCore/Source/PropertyAnimatorCoreEditor/Private/PropertyAnimatorCoreEditorModule.cpp
@@ -33,21 +33,24 @@ void FPropertyAnimatorCoreEditorModule::StartupModule()
 
 void FPropertyAnimatorCoreEditorModule::ShutdownModule()
 {
-	if (FPropertyEditorModule* PropertyModule = FModuleManager::GetModulePtr<FPropertyEditorModule>("PropertyEditor"))
+	if (UObjectInitialized() && !IsEngineExitRequested())
 	{
-		for (const FName& RegisteredCustomization : RegisteredCustomizations)
+		if (FPropertyEditorModule* PropertyModule = FModuleManager::GetModulePtr<FPropertyEditorModule>("PropertyEditor"))
 		{
-			PropertyModule->UnregisterCustomPropertyTypeLayout(RegisteredCustomization);
-			PropertyModule->UnregisterCustomClassLayout(RegisteredCustomization);
+			for (const FName& RegisteredCustomization : RegisteredCustomizations)
+			{
+				PropertyModule->UnregisterCustomPropertyTypeLayout(RegisteredCustomization);
+				PropertyModule->UnregisterCustomClassLayout(RegisteredCustomization);
+			}
 		}
-	}
-
-	RegisteredCustomizations.Empty();
 
-	if (ISequencerModule* SequencerModule = FModuleManager::GetModulePtr<ISequencerModule>("Sequencer"))
-	{
-		SequencerModule->UnRegisterTrackEditor(AnimatorTrackCreateEditorHandle);
-		AnimatorTrackCreateEditorHandle.Reset();
+		RegisteredCustomizations.Empty();
+		
+		if (ISequencerModule* SequencerModule = FModuleManager::GetModulePtr<ISequencerModule>("Sequencer"))
+		{
+			SequencerModule->UnRegisterTrackEditor(AnimatorTrackCreateEditorHandle);
+			AnimatorTrackCreateEditorHandle.Reset();
+		}
 	}
 }
 
